All "Open Source" actions on asset's source objects use the Edit verb when launching the file, so that the correctly bound Windows application can be used.  So that you get behaviors like photoshop opening textures instead of Image Previewer or Audacity for Sounds instead of Windows Media Player.

diff --git a/Engine/Source/Developer/AssetTools/Private/AssetTypeActions/AssetTypeActions_AnimSequence.cpp b/Engine/Source/Developer/AssetTools/Private/AssetTypeActions/AssetTypeActions_AnimSequence.cpp
--- a/Engine/Source/Developer/AssetTools/Private/AssetTypeActions/AssetTypeActions_AnimSequence.cpp
+++ b/Engine/Source/Developer/AssetTools/Private/AssetTypeActions/AssetTypeActions_AnimSequence.cpp
@@ -100,7 +100,7 @@ void FAssetTypeActions_AnimSequence::ExecuteOpenInExternalEditor(TArray<TWeakObj
 			const FString SourceFilePath = FReimportManager::ResolveImportFilename(Object->AssetImportData->SourceFilePath, Object);
 			if ( SourceFilePath.Len() && IFileManager::Get().FileSize( *SourceFilePath ) != INDEX_NONE )
 			{
-				FPlatformProcess::LaunchFileInDefaultExternalApplication( *SourceFilePath );
+				FPlatformProcess::LaunchFileInDefaultExternalApplication( *SourceFilePath, NULL, ELaunchVerb::Edit );
 			}
 		}
 	}
diff --git a/Engine/Source/Developer/AssetTools/Private/AssetTypeActions/AssetTypeActions_CSVAssetBase.cpp b/Engine/Source/Developer/AssetTools/Private/AssetTypeActions/AssetTypeActions_CSVAssetBase.cpp
--- a/Engine/Source/Developer/AssetTools/Private/AssetTypeActions/AssetTypeActions_CSVAssetBase.cpp
+++ b/Engine/Source/Developer/AssetTools/Private/AssetTypeActions/AssetTypeActions_CSVAssetBase.cpp
@@ -54,14 +54,14 @@ void FAssetTypeActions_CSVAssetBase::ExecuteOpenInExternalEditor(TArray<FString>
 				const FString FilenameWithExtension(FString::Printf(TEXT("%s/%s%s"), *RootPath, *BaseFilename, **ExtensionItr));
 				if (VerifyFileExists(FilenameWithExtension))
 				{
-					FPlatformProcess::LaunchFileInDefaultExternalApplication(*FilenameWithExtension);
+					FPlatformProcess::LaunchFileInDefaultExternalApplication(*FilenameWithExtension, NULL, ELaunchVerb::Edit);
 					break;
 				}
 			}
 		}
 		else if (VerifyFileExists(CSVFilename))
 		{
-			FPlatformProcess::LaunchFileInDefaultExternalApplication(*CSVFilename);
+			FPlatformProcess::LaunchFileInDefaultExternalApplication(*CSVFilename, NULL, ELaunchVerb::Edit);
 		}
 	}
 }
diff --git a/Engine/Source/Developer/AssetTools/Private/AssetTypeActions/AssetTypeActions_SkeletalMesh.cpp b/Engine/Source/Developer/AssetTools/Private/AssetTypeActions/AssetTypeActions_SkeletalMesh.cpp
--- a/Engine/Source/Developer/AssetTools/Private/AssetTypeActions/AssetTypeActions_SkeletalMesh.cpp
+++ b/Engine/Source/Developer/AssetTools/Private/AssetTypeActions/AssetTypeActions_SkeletalMesh.cpp
@@ -857,7 +857,7 @@ void FAssetTypeActions_SkeletalMesh::ExecuteOpenInExternalEditor(TArray<TWeakObj
 			const FString SourceFilePath = FReimportManager::ResolveImportFilename(Object->AssetImportData->SourceFilePath, Object);
 			if ( SourceFilePath.Len() && IFileManager::Get().FileSize( *SourceFilePath ) != INDEX_NONE )
 			{
-				FPlatformProcess::LaunchFileInDefaultExternalApplication( *SourceFilePath );
+				FPlatformProcess::LaunchFileInDefaultExternalApplication( *SourceFilePath, NULL, ELaunchVerb::Edit );
 			}
 		}
 	}
diff --git a/Engine/Source/Developer/AssetTools/Private/AssetTypeActions/AssetTypeActions_SoundWave.cpp b/Engine/Source/Developer/AssetTools/Private/AssetTypeActions/AssetTypeActions_SoundWave.cpp
--- a/Engine/Source/Developer/AssetTools/Private/AssetTypeActions/AssetTypeActions_SoundWave.cpp
+++ b/Engine/Source/Developer/AssetTools/Private/AssetTypeActions/AssetTypeActions_SoundWave.cpp
@@ -163,7 +163,7 @@ void FAssetTypeActions_SoundWave::ExecuteOpenInExternalEditor(TArray<TWeakObject
 			const FString SourceFilePath = FReimportManager::ResolveImportFilename(Object->SourceFilePath, Object);
 			if ( SourceFilePath.Len() && IFileManager::Get().FileSize( *SourceFilePath ) != INDEX_NONE )
 			{
-				FPlatformProcess::LaunchFileInDefaultExternalApplication( *SourceFilePath );
+				FPlatformProcess::LaunchFileInDefaultExternalApplication( *SourceFilePath, NULL, ELaunchVerb::Edit );
 			}
 		}
 	}
diff --git a/Engine/Source/Developer/AssetTools/Private/AssetTypeActions/AssetTypeActions_StaticMesh.cpp b/Engine/Source/Developer/AssetTools/Private/AssetTypeActions/AssetTypeActions_StaticMesh.cpp
--- a/Engine/Source/Developer/AssetTools/Private/AssetTypeActions/AssetTypeActions_StaticMesh.cpp
+++ b/Engine/Source/Developer/AssetTools/Private/AssetTypeActions/AssetTypeActions_StaticMesh.cpp
@@ -174,7 +174,7 @@ void FAssetTypeActions_StaticMesh::ExecuteOpenInExternalEditor(TArray<TWeakObjec
 			const FString SourceFilePath = FReimportManager::ResolveImportFilename(Object->AssetImportData->SourceFilePath, Object);
 			if ( SourceFilePath.Len() && IFileManager::Get().FileSize( *SourceFilePath ) != INDEX_NONE )
 			{
-				FPlatformProcess::LaunchFileInDefaultExternalApplication( *SourceFilePath );
+				FPlatformProcess::LaunchFileInDefaultExternalApplication( *SourceFilePath, NULL, ELaunchVerb::Edit );
 			}
 		}
 	}
diff --git a/Engine/Source/Developer/AssetTools/Private/AssetTypeActions/AssetTypeActions_Texture.cpp b/Engine/Source/Developer/AssetTools/Private/AssetTypeActions/AssetTypeActions_Texture.cpp
--- a/Engine/Source/Developer/AssetTools/Private/AssetTypeActions/AssetTypeActions_Texture.cpp
+++ b/Engine/Source/Developer/AssetTools/Private/AssetTypeActions/AssetTypeActions_Texture.cpp
@@ -202,7 +202,7 @@ void FAssetTypeActions_Texture::ExecuteOpenInExternalEditor(TArray<TWeakObjectPt
 			const FString SourceFilePath = FReimportManager::ResolveImportFilename(Object->SourceFilePath, Object);
 			if ( SourceFilePath.Len() && IFileManager::Get().FileSize( *SourceFilePath ) != INDEX_NONE )
 			{
-				FPlatformProcess::LaunchFileInDefaultExternalApplication( *SourceFilePath );
+				FPlatformProcess::LaunchFileInDefaultExternalApplication( *SourceFilePath, NULL, ELaunchVerb::Edit );
 			}
 		}
 	}
diff --git a/Engine/Source/Developer/AssetTools/Private/AssetTypeActions/AssetTypeActions_VectorFieldStatic.cpp b/Engine/Source/Developer/AssetTools/Private/AssetTypeActions/AssetTypeActions_VectorFieldStatic.cpp
--- a/Engine/Source/Developer/AssetTools/Private/AssetTypeActions/AssetTypeActions_VectorFieldStatic.cpp
+++ b/Engine/Source/Developer/AssetTools/Private/AssetTypeActions/AssetTypeActions_VectorFieldStatic.cpp
@@ -77,7 +77,7 @@ void FAssetTypeActions_VectorFieldStatic::ExecuteOpenInExternalEditor(TArray<TWe
 			const FString SourceFilePath = FReimportManager::ResolveImportFilename(Object->SourceFilePath, Object);
 			if ( SourceFilePath.Len() && IFileManager::Get().FileSize( *SourceFilePath ) != INDEX_NONE )
 			{
-				FPlatformProcess::LaunchFileInDefaultExternalApplication( *SourceFilePath );
+				FPlatformProcess::LaunchFileInDefaultExternalApplication( *SourceFilePath, NULL, ELaunchVerb::Edit );
 			}
 		}
 	}
