#ue4
- Added the ability to create a new anim blueprint from an existing one in the context menu of the content browser

diff --git a/Engine/Source/Developer/AssetTools/Private/AssetTypeActions/AssetTypeActions_AnimBlueprint.cpp b/Engine/Source/Developer/AssetTools/Private/AssetTypeActions/AssetTypeActions_AnimBlueprint.cpp
--- a/Engine/Source/Developer/AssetTools/Private/AssetTypeActions/AssetTypeActions_AnimBlueprint.cpp
+++ b/Engine/Source/Developer/AssetTools/Private/AssetTypeActions/AssetTypeActions_AnimBlueprint.cpp
@@ -112,12 +112,6 @@ void FAssetTypeActions_AnimBlueprint::OpenAssetEditor( const TArray<UObject*>& I
 	}
 }
 
-bool FAssetTypeActions_AnimBlueprint::CanCreateNewDerivedBlueprint() const
-{
-	// Animation blueprints cannot create derived blueprints.
-	return false;
-}
-
 void FAssetTypeActions_AnimBlueprint::ExecuteFindSkeleton(TArray<TWeakObjectPtr<UAnimBlueprint>> Objects)
 {
 	TArray<UObject*> ObjectsToSync;
diff --git a/Engine/Source/Developer/AssetTools/Private/AssetTypeActions/AssetTypeActions_AnimBlueprint.h b/Engine/Source/Developer/AssetTools/Private/AssetTypeActions/AssetTypeActions_AnimBlueprint.h
--- a/Engine/Source/Developer/AssetTools/Private/AssetTypeActions/AssetTypeActions_AnimBlueprint.h
+++ b/Engine/Source/Developer/AssetTools/Private/AssetTypeActions/AssetTypeActions_AnimBlueprint.h
@@ -13,11 +13,6 @@ public:
 	virtual void OpenAssetEditor( const TArray<UObject*>& InObjects, TSharedPtr<class IToolkitHost> EditWithinLevelEditor = TSharedPtr<IToolkitHost>() ) OVERRIDE;
 	virtual uint32 GetCategories() OVERRIDE { return EAssetTypeCategories::Animation; }
 
-protected:
-
-	/** Whether or not we can create derived blueprints */
-	virtual bool CanCreateNewDerivedBlueprint() const OVERRIDE;
-
 private:
 
 	/** Handler to fill the retarget submenu */
diff --git a/Engine/Source/Developer/AssetTools/Private/AssetTypeActions/AssetTypeActions_Blueprint.cpp b/Engine/Source/Developer/AssetTools/Private/AssetTypeActions/AssetTypeActions_Blueprint.cpp
--- a/Engine/Source/Developer/AssetTools/Private/AssetTypeActions/AssetTypeActions_Blueprint.cpp
+++ b/Engine/Source/Developer/AssetTools/Private/AssetTypeActions/AssetTypeActions_Blueprint.cpp
@@ -134,10 +134,10 @@ void FAssetTypeActions_Blueprint::ExecuteNewDerivedBlueprint(TWeakObjectPtr<UBlu
 	if ( Object )
 	{
 		// The menu option should ONLY be available if there is only one blueprint selected, validated by the menu creation code
-		UBlueprint* TargetBP = Object;
-		UClass* TargetClass = TargetBP->GeneratedClass;
+		UBlueprint* TargetParentBP = Object;
+		UClass* TargetParentClass = TargetParentBP->GeneratedClass;
 
-		if(!FKismetEditorUtilities::CanCreateBlueprintOfClass(TargetClass))
+		if (!FKismetEditorUtilities::CanCreateBlueprintOfClass(TargetParentClass))
 		{
 			FMessageDialog::Open( EAppMsgType::Ok, LOCTEXT("InvalidClassToMakeBlueprintFrom", "Invalid class with which to make a Blueprint."));
 			return;
@@ -148,11 +148,11 @@ void FAssetTypeActions_Blueprint::ExecuteNewDerivedBlueprint(TWeakObjectPtr<UBlu
 		CreateUniqueAssetName(Object->GetOutermost()->GetName(), TEXT("_Child"), PackageName, Name);
 
 		UPackage* Package = CreatePackage(NULL, *PackageName);
-		if ( ensure(Package) )
+		if (ensure(Package))
 		{
 			// Create and init a new Blueprint
-			UBlueprint* NewBP = FKismetEditorUtilities::CreateBlueprint(TargetClass, Package, FName(*Name), BPTYPE_Normal, UBlueprint::StaticClass());
-			if(NewBP)
+			UBlueprint* NewBP = FKismetEditorUtilities::CreateBlueprint(TargetParentClass, Package, FName(*Name), BPTYPE_Normal, TargetParentBP->GetClass());
+			if (NewBP)
 			{
 				FAssetEditorManager::Get().OpenEditorForAsset(NewBP);
 
diff --git a/Engine/Source/Editor/UnrealEd/Private/Kismet2/Kismet2.cpp b/Engine/Source/Editor/UnrealEd/Private/Kismet2/Kismet2.cpp
--- a/Engine/Source/Editor/UnrealEd/Private/Kismet2/Kismet2.cpp
+++ b/Engine/Source/Editor/UnrealEd/Private/Kismet2/Kismet2.cpp
@@ -240,11 +240,11 @@ UBlueprint* FKismetEditorUtilities::CreateBlueprint(UClass* ParentClass, UObject
 		CreateDefaultEventGraphs(NewBP);
 	}
 
-	// if this is an anim blueprint, add the root animation graph
 	//@TODO: ANIMREFACTOR: This kind of code should be on a per-blueprint basis; not centralized here
 	if (UAnimBlueprint* AnimBP = Cast<UAnimBlueprint>(NewBP))
 	{
-		if (UAnimBlueprint::FindRootAnimBlueprint(AnimBP) == NULL)
+		UAnimBlueprint* RootAnimBP = UAnimBlueprint::FindRootAnimBlueprint(AnimBP);
+		if (RootAnimBP == NULL)
 		{
 			// Only allow an anim graph if there isn't one in a parent blueprint
 			UEdGraph* NewGraph = FBlueprintEditorUtils::CreateNewGraph(AnimBP, K2Schema->GN_AnimGraph, UAnimationGraph::StaticClass(), UAnimationGraphSchema::StaticClass());
@@ -252,6 +252,11 @@ UBlueprint* FKismetEditorUtilities::CreateBlueprint(UClass* ParentClass, UObject
 			NewBP->LastEditedDocuments.Add(NewGraph);
 			NewGraph->bAllowDeletion = false;
 		}
+		else
+		{
+			// Make sure the anim blueprint targets the same skeleton as the parent
+			AnimBP->TargetSkeleton = RootAnimBP->TargetSkeleton;
+		}
 	}
 
 	// Create initial UClass
