UE4: Double clicking a world asset in the content browser now loads it in the level editor.

diff --git a/Engine/Source/Developer/AssetTools/Private/AssetTypeActions/AssetTypeActions_World.cpp b/Engine/Source/Developer/AssetTools/Private/AssetTypeActions/AssetTypeActions_World.cpp
--- a/Engine/Source/Developer/AssetTools/Private/AssetTypeActions/AssetTypeActions_World.cpp
+++ b/Engine/Source/Developer/AssetTools/Private/AssetTypeActions/AssetTypeActions_World.cpp
@@ -4,5 +4,23 @@
 
 #define LOCTEXT_NAMESPACE "AssetTypeActions"
 
+void FAssetTypeActions_World::OpenAssetEditor( const TArray<UObject*>& InObjects, TSharedPtr<class IToolkitHost> EditWithinLevelEditor )
+{
+	for (auto ObjIt = InObjects.CreateConstIterator(); ObjIt; ++ObjIt)
+	{
+		auto World = Cast<UWorld>(*ObjIt);
+		if (World != NULL)
+		{
+			const FString FileToOpen = FPackageName::LongPackageNameToFilename(World->GetOutermost()->GetName(), FPackageName::GetMapPackageExtension());
+			const bool bLoadAsTemplate = false;
+			const bool bShowProgress = false;
+			const bool bWorldComposition = false;
+			FEditorFileUtils::LoadMap( FileToOpen, bLoadAsTemplate, bShowProgress, bWorldComposition );
+
+			// We can only edit one world at a time... so just break after the first valid world to load
+			break;
+		}
+	}
+}
 
 #undef LOCTEXT_NAMESPACE
diff --git a/Engine/Source/Developer/AssetTools/Private/AssetTypeActions/AssetTypeActions_World.h b/Engine/Source/Developer/AssetTools/Private/AssetTypeActions/AssetTypeActions_World.h
--- a/Engine/Source/Developer/AssetTools/Private/AssetTypeActions/AssetTypeActions_World.h
+++ b/Engine/Source/Developer/AssetTools/Private/AssetTypeActions/AssetTypeActions_World.h
@@ -10,5 +10,6 @@ public:
 	virtual FColor GetTypeColor() const OVERRIDE { return FColor(255, 156, 0); }
 	virtual UClass* GetSupportedClass() const OVERRIDE { return UWorld::StaticClass(); }
 	virtual bool HasActions ( const TArray<UObject*>& InObjects ) const OVERRIDE { return false; }
+	virtual void OpenAssetEditor( const TArray<UObject*>& InObjects, TSharedPtr<class IToolkitHost> EditWithinLevelEditor = TSharedPtr<IToolkitHost>() ) OVERRIDE;
 	virtual uint32 GetCategories() OVERRIDE { return EAssetTypeCategories::Basic; }
 };
\ No newline at end of file
