L10N: Fixing localization conflict

diff --git a/Engine/Source/Developer/AssetTools/Private/AssetTypeActions/AssetTypeActions_Skeleton.cpp b/Engine/Source/Developer/AssetTools/Private/AssetTypeActions/AssetTypeActions_Skeleton.cpp
--- a/Engine/Source/Developer/AssetTools/Private/AssetTypeActions/AssetTypeActions_Skeleton.cpp
+++ b/Engine/Source/Developer/AssetTools/Private/AssetTypeActions/AssetTypeActions_Skeleton.cpp
@@ -31,8 +31,8 @@ void FAssetTypeActions_Skeleton::GetActions( const TArray<UObject*>& InObjects,
 
 	// create menu
 	MenuBuilder.AddSubMenu(
-			LOCTEXT("CreateAnimSubmenu", "Create"),
-			LOCTEXT("CreateAnimSubmenu_ToolTip", "Create assets for this skeleton"),
+			LOCTEXT("CreateSkeletonSubmenu", "Create"),
+			LOCTEXT("CreateSkeletonSubmenu_ToolTip", "Create assets for this skeleton"),
 			FNewMenuDelegate::CreateSP(this, &FAssetTypeActions_Skeleton::FillCreateMenu, Skeletons));
 
 	MenuBuilder.AddMenuEntry(
