Diff functionality now supports Paths with whitespace.
#TTP 331892 - Editor: Blueprint Defaults Diff: Lauching diff tool does not work with spaces in paths

#branch UE4

#proj Editor.BehaviourTreeEditor
#proj Editor.WorldBrowser
#proj Editor.Levels
#proj Developer.AssetTool

#change Wrapped the file paths in double quotes, this means that the full path gets passed to P4Merge

diff --git a/Engine/Source/Developer/AssetTools/Private/AssetTools.cpp b/Engine/Source/Developer/AssetTools/Private/AssetTools.cpp
--- a/Engine/Source/Developer/AssetTools/Private/AssetTools.cpp
+++ b/Engine/Source/Developer/AssetTools/Private/AssetTools.cpp
@@ -920,10 +920,21 @@ FString FAssetTools::DumpAssetToTempFile(UObject* Asset) const
 	}
 }
 
-bool FAssetTools::CreateDiffProcess(const FString& DiffCommand, const FString& DiffArgs) const
+FString WrapArgument(const FString& Argument)
 {
+	// Wrap the passed in argument so it changes from Argument to "Argument"
+	return FString::Printf(TEXT("%s%s%s"),	(Argument.StartsWith("\"")) ? "": "\"",
+											*Argument,
+											(Argument.EndsWith("\"")) ? "": "\"");
+}
+
+bool FAssetTools::CreateDiffProcess(const FString& DiffCommand,  const FString& OldTextFilename,  const FString& NewTextFilename, const FString& DiffArgs) const
+{
+	// Construct Arguments
+	FString Arguments = FString::Printf( TEXT("%s %s %s"),*WrapArgument(OldTextFilename), *WrapArgument(NewTextFilename), *DiffArgs );
+
 	// Fire process
-	if (!FPlatformProcess::CreateProc(*DiffCommand, *DiffArgs, true, false, false, NULL, 0, NULL, NULL).IsValid())
+	if (!FPlatformProcess::CreateProc(*DiffCommand, *Arguments, true, false, false, NULL, 0, NULL, NULL).IsValid())
 	{
 		//failed to work, put up error msg
 		FNotificationInfo Info( FText::Format( NSLOCTEXT("AssetTools", "DifFail", "Diff Failed: Could not launch external process '{0}'."), FText::FromString( DiffCommand ) ) );
diff --git a/Engine/Source/Developer/AssetTools/Private/AssetTools.h b/Engine/Source/Developer/AssetTools/Private/AssetTools.h
--- a/Engine/Source/Developer/AssetTools/Private/AssetTools.h
+++ b/Engine/Source/Developer/AssetTools/Private/AssetTools.h
@@ -24,7 +24,7 @@ public:
 	virtual void DiffAgainstDepot(UObject* InObject, const FString& InPackagePath, const FString& InPackageName) const OVERRIDE;
 	virtual void DiffAssets(UObject* OldAsset1, UObject* NewAsset, const struct FRevisionInfo& OldRevision, const struct FRevisionInfo& NewRevision) const OVERRIDE;
 	virtual FString DumpAssetToTempFile(UObject* Asset) const OVERRIDE;
-	virtual bool CreateDiffProcess(const FString& DiffCommand, const FString& DiffArgs) const OVERRIDE;
+	virtual bool CreateDiffProcess(const FString& DiffCommand, const FString& OldTextFilename, const FString& NewTextFilename, const FString& DiffArgs = FString("")) const OVERRIDE;
 	virtual void MigratePackages(const TArray<FName>& PackageNamesToMigrate) const OVERRIDE;
 	virtual void FixupReferencers(const TArray<UObjectRedirector*>& Objects) const OVERRIDE;
 
diff --git a/Engine/Source/Developer/AssetTools/Private/AssetTypeActions/AssetTypeActions_Blueprint.cpp b/Engine/Source/Developer/AssetTools/Private/AssetTypeActions/AssetTypeActions_Blueprint.cpp
--- a/Engine/Source/Developer/AssetTools/Private/AssetTypeActions/AssetTypeActions_Blueprint.cpp
+++ b/Engine/Source/Developer/AssetTools/Private/AssetTypeActions/AssetTypeActions_Blueprint.cpp
@@ -240,10 +240,8 @@ void FAssetTypeActions_Blueprint::OpenInDefaults( class UBlueprint* OldBlueprint
 		const FString NewTextFilename = DumpAssetToTempFile(NewBlueprint->GeneratedClass->GetDefaultObject());
 		const FString DiffCommand = GetDefault<UEditorLoadingSavingSettings>()->TextDiffToolPath.FilePath;
 
-		// args are just 2 temp filenames
-		const FString DiffArgs = FString::Printf(TEXT("%s %s"), *OldTextFilename, *NewTextFilename);
-
-		CreateDiffProcess(DiffCommand, DiffArgs);
+		FAssetToolsModule& AssetToolsModule = FModuleManager::Get().LoadModuleChecked<FAssetToolsModule>("AssetTools");
+     	AssetToolsModule.Get().CreateDiffProcess(DiffCommand, OldTextFilename, NewTextFilename);
 	}
 }
 
diff --git a/Engine/Source/Developer/AssetTools/Private/AssetTypeActions/AssetTypeActions_BlueprintGeneratedClass.cpp b/Engine/Source/Developer/AssetTools/Private/AssetTypeActions/AssetTypeActions_BlueprintGeneratedClass.cpp
--- a/Engine/Source/Developer/AssetTools/Private/AssetTypeActions/AssetTypeActions_BlueprintGeneratedClass.cpp
+++ b/Engine/Source/Developer/AssetTools/Private/AssetTypeActions/AssetTypeActions_BlueprintGeneratedClass.cpp
@@ -217,10 +217,8 @@ void FAssetTypeActions_BlueprintGeneratedClass::OpenInDefaults( class UBlueprint
 	FString NewTextFilename = DumpAssetToTempFile(NewBlueprint->GeneratedClass->GetDefaultObject());
 	FString DiffCommand = GetDefault<UEditorLoadingSavingSettings>()->TextDiffToolPath.FilePath;
 
-	// args are just 2 temp filenames
-	FString DiffArgs = FString::Printf(TEXT("%s %s"), *OldTextFilename, *NewTextFilename);
-
-	CreateDiffProcess(DiffCommand, DiffArgs);
+	FAssetToolsModule& AssetToolsModule = FModuleManager::Get().LoadModuleChecked<FAssetToolsModule>("AssetTools");
+	AssetToolsModule.Get().CreateDiffProcess(DiffCommand, OldTextFilename, NewTextFilename);
 }
 
 #undef LOCTEXT_NAMESPACE
diff --git a/Engine/Source/Developer/AssetTools/Public/AssetTypeActions_Base.h b/Engine/Source/Developer/AssetTools/Public/AssetTypeActions_Base.h
--- a/Engine/Source/Developer/AssetTools/Public/AssetTypeActions_Base.h
+++ b/Engine/Source/Developer/AssetTools/Public/AssetTypeActions_Base.h
@@ -40,12 +40,12 @@ public:
 		}
 	}
 
-	virtual bool CanFilter() 
+	virtual bool CanFilter() OVERRIDE
 	{
 		return true;
 	}
 
-	virtual bool ShouldForceWorldCentric()
+	virtual bool ShouldForceWorldCentric() OVERRIDE
 	{
 		return false;
 	}
@@ -60,10 +60,8 @@ public:
 		FString NewTextFilename = DumpAssetToTempFile(NewAsset);
 		FString DiffCommand = GetDefault<UEditorLoadingSavingSettings>()->TextDiffToolPath.FilePath;
 
-		// args are just 2 temp filenames
-		FString DiffArgs = FString::Printf(TEXT("%s %s"), *OldTextFilename, *NewTextFilename);
-
-		CreateDiffProcess(DiffCommand, DiffArgs);
+		FAssetToolsModule& AssetToolsModule = FModuleManager::Get().LoadModuleChecked<FAssetToolsModule>("AssetTools");
+		AssetToolsModule.Get().CreateDiffProcess(DiffCommand, OldTextFilename, NewTextFilename);
 	}
 
 	virtual class UThumbnailInfo* GetThumbnailInfo(UObject* Asset) const OVERRIDE
@@ -91,13 +89,6 @@ protected:
 		return AssetToolsModule.Get().DumpAssetToTempFile(Asset);
 	}
 
-	/* Attempt to spawn diff tool as external process */
-	virtual bool CreateDiffProcess(const FString& DiffCommand, const FString& DiffArgs) const
-	{
-		FAssetToolsModule& AssetToolsModule = FModuleManager::Get().LoadModuleChecked<FAssetToolsModule>("AssetTools");
-		return AssetToolsModule.Get().CreateDiffProcess(DiffCommand, DiffArgs);
-	}
-
 	/** Returns additional tooltip information for the specified asset, if it has any (otherwise return the null widget) */
 	virtual FText GetAssetDescription(const FAssetData& AssetData) const
 	{
diff --git a/Engine/Source/Developer/AssetTools/Public/IAssetTools.h b/Engine/Source/Developer/AssetTools/Public/IAssetTools.h
--- a/Engine/Source/Developer/AssetTools/Public/IAssetTools.h
+++ b/Engine/Source/Developer/AssetTools/Public/IAssetTools.h
@@ -92,8 +92,16 @@ public:
 	/** Util for dumping an asset to a temporary text file. Returns absolute filename to temp file */
 	virtual FString DumpAssetToTempFile(UObject* Asset) const = 0;
 
-	/* Attempt to spawn diff tool as external process */
-	virtual bool CreateDiffProcess(const FString& DiffCommand, const FString& DiffArgs) const = 0;
+	/* Attempt to spawn Diff tool as external process 
+	 *
+	 * @param DiffCommand -		Command used to launch the diff tool
+	 * @param OldTextFilename - File path to original file 
+	 * @param NewTextFilename - File path to new file
+	 * @param DiffArgs -		Any extra command line arguments (defaulted to empty)
+	 *
+	 * @return Returns true if the process has successfully been created.
+	 */
+	virtual bool CreateDiffProcess(const FString& DiffCommand, const FString& OldTextFilename, const FString& NewTextFilename, const FString& DiffArgs = FString("")) const = 0;
 
 	/* Migrate packages to another game content folder */
 	virtual void MigratePackages(const TArray<FName>& PackageNamesToMigrate) const = 0;
diff --git a/Engine/Source/Editor/BehaviorTreeEditor/Private/AssetTypeActions_BehaviorTree.cpp b/Engine/Source/Editor/BehaviorTreeEditor/Private/AssetTypeActions_BehaviorTree.cpp
--- a/Engine/Source/Editor/BehaviorTreeEditor/Private/AssetTypeActions_BehaviorTree.cpp
+++ b/Engine/Source/Editor/BehaviorTreeEditor/Private/AssetTypeActions_BehaviorTree.cpp
@@ -124,10 +124,9 @@ void FAssetTypeActions_BehaviorTree::OpenInDefaults( class UBehaviorTree* OldBeh
 
 	// Get diff program to use
 	FString DiffCommand = GetDefault<UEditorLoadingSavingSettings>()->TextDiffToolPath.FilePath;
-	// args are just 2 temp filenames
-	FString DiffArgs = FString::Printf(TEXT("%s %s"), *OldTextFilename, *NewTextFilename);
 
-	CreateDiffProcess(DiffCommand, DiffArgs);
+	FAssetToolsModule& AssetToolsModule = FModuleManager::Get().LoadModuleChecked<FAssetToolsModule>("AssetTools");
+	AssetToolsModule.Get().CreateDiffProcess(DiffCommand, OldTextFilename, NewTextFilename);
 }
 
 #undef LOCTEXT_NAMESPACE
diff --git a/Engine/Source/Editor/Levels/Private/LevelCollectionViewModel.cpp b/Engine/Source/Editor/Levels/Private/LevelCollectionViewModel.cpp
--- a/Engine/Source/Editor/Levels/Private/LevelCollectionViewModel.cpp
+++ b/Engine/Source/Editor/Levels/Private/LevelCollectionViewModel.cpp
@@ -1306,10 +1306,7 @@ void FLevelCollectionViewModel::OnSCCDiffAgainstDepot()
 							FString NewTextFilename = AssetToolsModule.Get().DumpAssetToTempFile(OriginalPackage);
 							FString DiffCommand = GetDefault<UEditorLoadingSavingSettings>()->TextDiffToolPath.FilePath;
 
-							// args are just 2 temp filenames
-							FString DiffArgs = FString::Printf(TEXT("%s %s"), *OldTextFilename, *NewTextFilename);
-
-							AssetToolsModule.Get().CreateDiffProcess(DiffCommand, DiffArgs);
+							AssetToolsModule.Get().CreateDiffProcess(DiffCommand, OldTextFilename, NewTextFilename);
 
 							AssetToolsModule.Get().DiffAssets(OldPackage, OriginalPackage, OldRevision, NewRevision);
 						}
diff --git a/Engine/Source/Editor/WorldBrowser/Private/LevelCollectionModel.cpp b/Engine/Source/Editor/WorldBrowser/Private/LevelCollectionModel.cpp
--- a/Engine/Source/Editor/WorldBrowser/Private/LevelCollectionModel.cpp
+++ b/Engine/Source/Editor/WorldBrowser/Private/LevelCollectionModel.cpp
@@ -1119,10 +1119,7 @@ void FLevelCollectionModel::SCCDiffAgainstDepot(const FLevelModelList& InList, U
 							FString NewTextFilename = AssetToolsModule.Get().DumpAssetToTempFile(OriginalPackage);
 							FString DiffCommand = GetDefault<UEditorLoadingSavingSettings>()->TextDiffToolPath.FilePath;
 
-							// args are just 2 temp filenames
-							FString DiffArgs = FString::Printf(TEXT("%s %s"), *OldTextFilename, *NewTextFilename);
-
-							AssetToolsModule.Get().CreateDiffProcess(DiffCommand, DiffArgs);
+							AssetToolsModule.Get().CreateDiffProcess(DiffCommand, OldTextFilename, NewTextFilename);
 							AssetToolsModule.Get().DiffAssets(OldPackage, OriginalPackage, OldRevision, NewRevision);
 						}
 					}
