Fixed a crash caused by a callback being removed from a multicast delegate while it was still pending invocation

This was a tough one to track down. The problem showed up when closing a matinee editor whilst a preview camera was selected in the world. Several clients all hook into the EditorModeChanged event to do some cleanup when the matinee editor is closed.

When matinee was closed, it fired the event to say that the mode is changing. This broadcast took a copy of the invocation list, and proceeded to invoke each delegate. One of these handlers resulted in the preview actor being deleted, which ended up deleting its associated preview viewport. As part of the viewport's destructor it correctly removes its raw ptr from the mode changing delegate list. However, since this event is currently broadcasting, a copy of the delegate still exists in the local invocation list which is invoked later on. The trouble is, the user object has been destroyed by this time.

I have fixed this by moving the destruction of the preview actor to a subsequent event triggered when the matinee mode is exited. This is invoked immediately after the OnEditorModeChanged event, and fixes the issue because it waits until the first broadcast to finish before destroying the viewport.

The alternative would have been to make FEditorViewportClient derive from TSharedFromThis<> and use AddSP for its binding, but I don't think that publicizing the sharing of viewport clients is the right approach. I'd rather keep ownership of viewport clients more explicit.

diff --git a/Engine/Source/Developer/AssetTools/Private/AssetTypeActions/AssetTypeActions_CameraAnim.cpp b/Engine/Source/Developer/AssetTools/Private/AssetTypeActions/AssetTypeActions_CameraAnim.cpp
--- a/Engine/Source/Developer/AssetTools/Private/AssetTypeActions/AssetTypeActions_CameraAnim.cpp
+++ b/Engine/Source/Developer/AssetTools/Private/AssetTypeActions/AssetTypeActions_CameraAnim.cpp
@@ -212,7 +212,7 @@ void FAssetTypeActions_CameraAnim::OpenAssetEditor( const TArray<UObject*>& InOb
 				GEditor->OpenMatinee(PreviewMatineeActor.Get());
 
 				// install our delegate so we can clean up when finished
-				GEditorModeTools().OnEditorModeChanged().AddRaw(this, &FAssetTypeActions_CameraAnim::OnMatineeEditorClosed);
+				FEditorDelegates::EditorModeExit.AddSP(this, &FAssetTypeActions_CameraAnim::OnMatineeEditorClosed);
 			}
 			else
 			{
@@ -222,10 +222,10 @@ void FAssetTypeActions_CameraAnim::OpenAssetEditor( const TArray<UObject*>& InOb
 	}
 }
 
-void FAssetTypeActions_CameraAnim::OnMatineeEditorClosed(FEdMode* InEditorMode, bool bIsEnteringMode)
+void FAssetTypeActions_CameraAnim::OnMatineeEditorClosed(FEdMode* InEditorMode)
 {
 	check(InEditorMode);
-	if(InEditorMode->GetID() == FBuiltinEditorModes::EM_InterpEdit && !bIsEnteringMode)
+	if(InEditorMode->GetID() == FBuiltinEditorModes::EM_InterpEdit)
 	{
 		// clean up our preview actors if they are still present
 		if(PreviewCamera.IsValid())
@@ -247,7 +247,7 @@ void FAssetTypeActions_CameraAnim::OnMatineeEditorClosed(FEdMode* InEditorMode,
 		}
 
 		// remove our delegate
-		GEditorModeTools().OnEditorModeChanged().RemoveRaw(this, &FAssetTypeActions_CameraAnim::OnMatineeEditorClosed);
+		FEditorDelegates::EditorModeExit.RemoveSP(this, &FAssetTypeActions_CameraAnim::OnMatineeEditorClosed);
 	}
 }
 
diff --git a/Engine/Source/Developer/AssetTools/Private/AssetTypeActions/AssetTypeActions_CameraAnim.h b/Engine/Source/Developer/AssetTools/Private/AssetTypeActions/AssetTypeActions_CameraAnim.h
--- a/Engine/Source/Developer/AssetTools/Private/AssetTypeActions/AssetTypeActions_CameraAnim.h
+++ b/Engine/Source/Developer/AssetTools/Private/AssetTypeActions/AssetTypeActions_CameraAnim.h
@@ -21,7 +21,7 @@ private:
 	void CreateCameraActorForCameraAnim(class UCameraAnim* InCameraAnim);
 
 	/** Delegate fired when the editor mode is changed */
-	void OnMatineeEditorClosed(class FEdMode* InEditorMode, bool bIsEnteringMode);
+	void OnMatineeEditorClosed(class FEdMode* InEditorMode);
 
 	/** Helper function to hookup preview pawn */
 	void CreatePreviewPawnForCameraAnim(class UCameraAnim* InCameraAnim);
diff --git a/Engine/Source/Editor/UnrealEd/Private/Editor.cpp b/Engine/Source/Editor/UnrealEd/Private/Editor.cpp
--- a/Engine/Source/Editor/UnrealEd/Private/Editor.cpp
+++ b/Engine/Source/Editor/UnrealEd/Private/Editor.cpp
@@ -88,8 +88,8 @@ FSimpleMulticastDelegate								FEditorDelegates::RefreshLayerBrowser;
 FSimpleMulticastDelegate								FEditorDelegates::RefreshPrimitiveStatsBrowser;
 FSimpleMulticastDelegate								FEditorDelegates::LoadSelectedAssetsIfNeeded;
 FSimpleMulticastDelegate								FEditorDelegates::DisplayLoadErrors;
-FSimpleMulticastDelegate								FEditorDelegates::EditorModeEnter;
-FSimpleMulticastDelegate								FEditorDelegates::EditorModeExit;
+FEditorDelegates::FOnEditorModeTransitioned				FEditorDelegates::EditorModeEnter;
+FEditorDelegates::FOnEditorModeTransitioned				FEditorDelegates::EditorModeExit;
 FSimpleMulticastDelegate								FEditorDelegates::Undo;
 FEditorDelegates::FOnPIEEvent							FEditorDelegates::BeginPIE;
 FEditorDelegates::FOnPIEEvent							FEditorDelegates::EndPIE;
diff --git a/Engine/Source/Editor/UnrealEd/Private/EditorModes.cpp b/Engine/Source/Editor/UnrealEd/Private/EditorModes.cpp
--- a/Engine/Source/Editor/UnrealEd/Private/EditorModes.cpp
+++ b/Engine/Source/Editor/UnrealEd/Private/EditorModes.cpp
@@ -556,7 +556,7 @@ void FEdMode::Enter()
 		SelectedActor->MarkComponentsRenderStateDirty();
 	}
 
-	FEditorDelegates::EditorModeEnter.Broadcast();
+	FEditorDelegates::EditorModeEnter.Broadcast( this );
 	const bool bIsEnteringMode = true;
 	GEditorModeTools().BroadcastEditorModeChanged( this, bIsEnteringMode );
 }
@@ -565,7 +565,7 @@ void FEdMode::Exit()
 {
 	const bool bIsEnteringMode = false;
 	GEditorModeTools().BroadcastEditorModeChanged( this, bIsEnteringMode );
-	FEditorDelegates::EditorModeExit.Broadcast();
+	FEditorDelegates::EditorModeExit.Broadcast( this );
 }
 
 void FEdMode::SetCurrentTool( EModeTools InID )
diff --git a/Engine/Source/Editor/UnrealEd/Private/EditorViewportClient.cpp b/Engine/Source/Editor/UnrealEd/Private/EditorViewportClient.cpp
--- a/Engine/Source/Editor/UnrealEd/Private/EditorViewportClient.cpp
+++ b/Engine/Source/Editor/UnrealEd/Private/EditorViewportClient.cpp
@@ -311,7 +311,10 @@ void FEditorViewportClient::RequestInvalidateHitProxy(FViewport* InViewport)
 
 void FEditorViewportClient::OnEditorModeChanged(FEdMode* EditorMode, bool bIsEntering)
 {
-	RequestInvalidateHitProxy(Viewport);
+	if (Viewport)
+	{
+		RequestInvalidateHitProxy(Viewport);
+	}
 }
 
 float FEditorViewportClient::GetOrthoUnitsPerPixel(const FViewport* InViewport) const
diff --git a/Engine/Source/Editor/UnrealEd/Public/Editor.h b/Engine/Source/Editor/UnrealEd/Public/Editor.h
--- a/Engine/Source/Editor/UnrealEd/Public/Editor.h
+++ b/Engine/Source/Editor/UnrealEd/Public/Editor.h
@@ -83,6 +83,8 @@ struct UNREALED_API FEditorDelegates
 	DECLARE_MULTICAST_DELEGATE_TwoParams(FOnMapOpened, const FString& /* Filename */, bool /*bAsTemplate*/);
 	/** delegate type for triggering when the Blueprint ContextMenu is created */
 	DECLARE_MULTICAST_DELEGATE_OneParam(FOnBlueprintContextMenuCreated, FBlueprintGraphActionListBuilder &/*ContextMenuBuilder*/);
+	/** Delegate used for entering or exiting an editor mode */
+	DECLARE_MULTICAST_DELEGATE_OneParam(FOnEditorModeTransitioned, FEdMode* /*Mode*/);
 
 	/** Called when the CurrentLevel is switched to a new level.  Note that this event won't be fired for temporary
 		changes to the current level, such as when copying/pasting actors. */
@@ -117,9 +119,9 @@ struct UNREALED_API FEditorDelegates
 	/** Called when load errors are about to be displayed */
 	static FSimpleMulticastDelegate DisplayLoadErrors;
 	/** Called when an editor mode is being entered */
-	static FSimpleMulticastDelegate EditorModeEnter;
+	static FOnEditorModeTransitioned EditorModeEnter;
 	/** Called when an editor mode is being exited */
-	static FSimpleMulticastDelegate EditorModeExit;
+	static FOnEditorModeTransitioned EditorModeExit;
 	/** Sent after an undo/redo operation takes place */
 	static FSimpleMulticastDelegate Undo;
 	/** Sent when a PIE session is beginning */
