UE4: When opening a world via the content browser, you are now prompted for save if you have unsaved work, just as you would be if you opened via the file menu.

diff --git a/Engine/Source/Developer/AssetTools/Private/AssetTypeActions/AssetTypeActions_World.cpp b/Engine/Source/Developer/AssetTools/Private/AssetTypeActions/AssetTypeActions_World.cpp
--- a/Engine/Source/Developer/AssetTools/Private/AssetTypeActions/AssetTypeActions_World.cpp
+++ b/Engine/Source/Developer/AssetTools/Private/AssetTypeActions/AssetTypeActions_World.cpp
@@ -11,10 +11,17 @@ void FAssetTypeActions_World::OpenAssetEditor( const TArray<UObject*>& InObjects
 		auto World = Cast<UWorld>(*ObjIt);
 		if (World != NULL && World != GEditor->GetEditorWorldContext().World())
 		{
-			const FString FileToOpen = FPackageName::LongPackageNameToFilename(World->GetOutermost()->GetName(), FPackageName::GetMapPackageExtension());
-			const bool bLoadAsTemplate = false;
-			const bool bShowProgress = false;
-			FEditorFileUtils::LoadMap( FileToOpen, bLoadAsTemplate, bShowProgress );
+			// If there are any unsaved changes to the current level, see if the user wants to save those first.
+			bool bPromptUserToSave = true;
+			bool bSaveMapPackages = true;
+			bool bSaveContentPackages = true;
+			if (FEditorFileUtils::SaveDirtyPackages(bPromptUserToSave, bSaveMapPackages, bSaveContentPackages))
+			{
+				const FString FileToOpen = FPackageName::LongPackageNameToFilename(World->GetOutermost()->GetName(), FPackageName::GetMapPackageExtension());
+				const bool bLoadAsTemplate = false;
+				const bool bShowProgress = true;
+				FEditorFileUtils::LoadMap(FileToOpen, bLoadAsTemplate, bShowProgress);
+			}
 
 			// We can only edit one world at a time... so just break after the first valid world to load
 			break;
