UE4: World assets are now shown as "Level" assets in the UI.

diff --git a/Engine/Source/Developer/AssetTools/Private/AssetTypeActions/AssetTypeActions_World.h b/Engine/Source/Developer/AssetTools/Private/AssetTypeActions/AssetTypeActions_World.h
--- a/Engine/Source/Developer/AssetTools/Private/AssetTypeActions/AssetTypeActions_World.h
+++ b/Engine/Source/Developer/AssetTools/Private/AssetTypeActions/AssetTypeActions_World.h
@@ -6,7 +6,7 @@ class FAssetTypeActions_World : public FAssetTypeActions_Base
 {
 public:
 	// IAssetTypeActions Implementation
-	virtual FText GetName() const OVERRIDE { return NSLOCTEXT("AssetTypeActions", "AssetTypeActions_World", "World"); }
+	virtual FText GetName() const OVERRIDE { return NSLOCTEXT("AssetTypeActions", "AssetTypeActions_World", "Level"); }
 	virtual FColor GetTypeColor() const OVERRIDE { return FColor(255, 156, 0); }
 	virtual UClass* GetSupportedClass() const OVERRIDE { return UWorld::StaticClass(); }
 	virtual bool HasActions ( const TArray<UObject*>& InObjects ) const OVERRIDE { return false; }
