#UE4 adding an accessor so that I don't have to copy/paste the loadmodule call, matches the empty module template...

diff --git a/Engine/Source/Developer/AssetTools/Public/AssetToolsModule.h b/Engine/Source/Developer/AssetTools/Public/AssetToolsModule.h
--- a/Engine/Source/Developer/AssetTools/Public/AssetToolsModule.h
+++ b/Engine/Source/Developer/AssetTools/Public/AssetToolsModule.h
@@ -14,6 +14,17 @@ public:
 
 	virtual IAssetTools& Get() const;
 
+	/**
+	 * Singleton-like access to this module's interface.  This is just for convenience!
+	 * Beware of calling this during the shutdown phase, though.  Your module might have been unloaded already.
+	 *
+	 * @return Returns singleton instance, loading the module on demand if needed
+	 */
+	static inline FAssetToolsModule& GetModule()
+	{
+		return FModuleManager::LoadModuleChecked< FAssetToolsModule >("AssetTools");
+	}
+
 private:
 	IAssetTools* AssetTools;
 	class FAssetToolsConsoleCommands* ConsoleCommands;
