Initial support for other foliage types
Note: The FoliageMeshes map is now keyed on a UFoliageType*, not a UStaticMesh*!

diff --git a/Engine/Config/BaseEngine.ini b/Engine/Config/BaseEngine.ini
--- a/Engine/Config/BaseEngine.ini
+++ b/Engine/Config/BaseEngine.ini
@@ -1230,8 +1230,11 @@ bUseStreamingPause=false
 +ActiveClassRedirects=(OldClassName="SlateWidgetStyleAsset", NewClassName="/Script/SlateCore.SlateWidgetStyleAsset")
 +ActiveClassRedirects=(OldClassName="SlateWidgetStyleContainerBase", NewClassName="/Script/SlateCore.SlateWidgetStyleContainerBase")
 
+;Exposing SplineMeshComponents to Blueprint
 +TaggedPropertyRedirects=(ClassName="SplineMeshComponent",OldPropertyName="SplineXDir",NewPropertyName="SplineUpDir")
 
+;Foliage refactor
++ActiveClassRedirects=(OldClassName="InstancedFoliageSettings", NewClassName="FoliageType_InstancedStaticMesh")
 
 
 [CoreUObject.Metadata]
diff --git a/Engine/Source/Developer/AssetTools/Private/AssetTypeActions/AssetTypeActions_InstancedFoliageSettings.cpp b/Engine/Source/Developer/AssetTools/Private/AssetTypeActions/AssetTypeActions_InstancedFoliageSettings.cpp
--- a/Engine/Source/Developer/AssetTools/Private/AssetTypeActions/AssetTypeActions_InstancedFoliageSettings.cpp
+++ b/Engine/Source/Developer/AssetTools/Private/AssetTypeActions/AssetTypeActions_InstancedFoliageSettings.cpp
@@ -1,10 +1,9 @@
 // Copyright 1998-2014 Epic Games, Inc. All Rights Reserved.
 
 #include "AssetToolsPrivatePCH.h"
-#include "Foliage/InstancedFoliageSettings.h"
+#include "Foliage/FoliageType_InstancedStaticMesh.h"
 
 UClass* FAssetTypeActions_InstancedFoliageSettings::GetSupportedClass() const
 {
-	// UInstancedFoliageSettings is defined in EngineFoliageClasses.h
-	return UInstancedFoliageSettings::StaticClass();
+	return UFoliageType_InstancedStaticMesh::StaticClass();
 }
\ No newline at end of file
