#UE4 I will not use uniform initialization syntax, I will not use uniform initialization syntax, I will not use uniform initialization syntax

diff --git a/Engine/Source/Developer/AssetTools/Public/IAssetTypeActions.h b/Engine/Source/Developer/AssetTools/Public/IAssetTypeActions.h
--- a/Engine/Source/Developer/AssetTools/Public/IAssetTypeActions.h
+++ b/Engine/Source/Developer/AssetTools/Public/IAssetTypeActions.h
@@ -39,12 +39,10 @@ struct FRevisionInfo
 	int32		Changelist;
 	FDateTime	Date;	
 
-	static FRevisionInfo InvalidRevision()
+	static inline FRevisionInfo InvalidRevision()
 	{
-		static const FRevisionInfo InvalidRevision{
-			-1, -1, FDateTime()
-		};
-		return InvalidRevision;
+		static const FRevisionInfo Ret = { -1, -1, FDateTime() };
+		return Ret;
 	}
 };
 
