Fixed FBX issue where failed once, the second time it crashes due to invalid object.

343098	EDITOR: FBX: CRASH: Editor will crash when trying to import a skeletal mesh with multiple root bones a second time

- It was due to package is deleted, but access it again after.

diff --git a/Engine/Source/Developer/AssetTools/Private/AssetTools.cpp b/Engine/Source/Developer/AssetTools/Private/AssetTools.cpp
--- a/Engine/Source/Developer/AssetTools/Private/AssetTools.cpp
+++ b/Engine/Source/Developer/AssetTools/Private/AssetTools.cpp
@@ -699,6 +699,11 @@ TArray<UObject*> FAssetTools::ImportAssets(const TArray<FString>& Files, const F
 							OnNewImportRecord(ImportAssetType, FileExtension, bImportSucceeded, bImportWasCancelled, ImportStartTime);
 							continue;
 						}
+						else
+						{
+							// succeed, recreate package since it has been deleted
+							Pkg = CreatePackage(NULL, *PackageName);
+						}
 					}
 					else
 					{
