ttp343706 CRASH:EDITOR: Editor Crashes when attempting to open BP_Soldier (User asset)

diff --git a/Engine/Source/Developer/AssetTools/Private/AssetTypeActions/AssetTypeActions_Blueprint.cpp b/Engine/Source/Developer/AssetTools/Private/AssetTypeActions/AssetTypeActions_Blueprint.cpp
--- a/Engine/Source/Developer/AssetTools/Private/AssetTypeActions/AssetTypeActions_Blueprint.cpp
+++ b/Engine/Source/Developer/AssetTools/Private/AssetTypeActions/AssetTypeActions_Blueprint.cpp
@@ -172,7 +172,7 @@ void FAssetTypeActions_Blueprint::ExecuteNewDerivedBlueprint(TWeakObjectPtr<UBlu
 
 FText FAssetTypeActions_Blueprint::GetNewDerivedBlueprintTooltip(TWeakObjectPtr<UBlueprint> InObject)
 {
-	if(InObject->GeneratedClass->HasAnyClassFlags(CLASS_Deprecated))
+	if (!CanExecuteNewDerivedBlueprint(InObject))
 	{
 		return LOCTEXT("Blueprint_NewDerivedBlueprintIsDeprecatedTooltip", "Blueprint class is deprecated, cannot derive a child Blueprint!");
 	}
@@ -184,7 +184,9 @@ FText FAssetTypeActions_Blueprint::GetNewDerivedBlueprintTooltip(TWeakObjectPtr<
 
 bool FAssetTypeActions_Blueprint::CanExecuteNewDerivedBlueprint(TWeakObjectPtr<UBlueprint> InObject)
 {
-	return !InObject->GeneratedClass->HasAnyClassFlags(CLASS_Deprecated);
+	auto BP = InObject.Get();
+	auto BPGC = BP ? BP->GeneratedClass : NULL;
+	return BPGC && !BPGC->HasAnyClassFlags(CLASS_Deprecated);
 }
 
 bool FAssetTypeActions_Blueprint::ShouldUseDataOnlyEditor( const UBlueprint* Blueprint ) const
