Fixing crash with new blueprint menu system; ensuring that new blueprints are added to the asset registry before being opened in the editor.

diff --git a/Engine/Source/Developer/AssetTools/Private/AssetTypeActions/AssetTypeActions_Blueprint.cpp b/Engine/Source/Developer/AssetTools/Private/AssetTypeActions/AssetTypeActions_Blueprint.cpp
--- a/Engine/Source/Developer/AssetTools/Private/AssetTypeActions/AssetTypeActions_Blueprint.cpp
+++ b/Engine/Source/Developer/AssetTools/Private/AssetTypeActions/AssetTypeActions_Blueprint.cpp
@@ -158,11 +158,14 @@ void FAssetTypeActions_Blueprint::ExecuteNewDerivedBlueprint(TWeakObjectPtr<UBlu
 			UBlueprint* NewBP = FKismetEditorUtilities::CreateBlueprint(TargetParentClass, Package, FName(*Name), BPTYPE_Normal, TargetParentBP->GetClass(), UBlueprintGeneratedClass::StaticClass());
 			if (NewBP)
 			{
-				FAssetEditorManager::Get().OpenEditorForAsset(NewBP);
-
 				// Notify the asset registry
 				FAssetRegistryModule::AssetCreated(NewBP);
 
+				// the editor should be opened AFTER being added to the asset 
+				// registry (some systems could queue off of the asset registry 
+				// add event, and already having that blueprint open can be odd)
+				FAssetEditorManager::Get().OpenEditorForAsset(NewBP);
+
 				// Mark the package dirty...
 				Package->MarkPackageDirty();
 			}
