* nicer subsurface profile asset color

diff --git a/Engine/Source/Developer/AssetTools/Private/AssetTypeActions/AssetTypeActions_SubsurfaceProfile.h b/Engine/Source/Developer/AssetTools/Private/AssetTypeActions/AssetTypeActions_SubsurfaceProfile.h
--- a/Engine/Source/Developer/AssetTools/Private/AssetTypeActions/AssetTypeActions_SubsurfaceProfile.h
+++ b/Engine/Source/Developer/AssetTools/Private/AssetTypeActions/AssetTypeActions_SubsurfaceProfile.h
@@ -7,7 +7,7 @@ class FAssetTypeActions_SubsurfaceProfile : public FAssetTypeActions_Base
 public:
 	// IAssetTypeActions Implementation
 	virtual FText GetName() const override { return NSLOCTEXT("AssetTypeActions", "AssetTypeActions_SubsurfaceProfile", "Subsurface Profile"); }
-	virtual FColor GetTypeColor() const override { return FColor(192, 121, 64); }
+	virtual FColor GetTypeColor() const override { return FColor(255, 196, 128); }
 	virtual UClass* GetSupportedClass() const override;
 
 	virtual uint32 GetCategories() override { return EAssetTypeCategories::MaterialsAndTextures; }
