Fixed crash happening with curves when converting from one skeleton to another skeleton

diff --git a/Engine/Source/Developer/AssetTools/Private/AssetTypeActions/AssetTypeActions_Skeleton.cpp b/Engine/Source/Developer/AssetTools/Private/AssetTypeActions/AssetTypeActions_Skeleton.cpp
--- a/Engine/Source/Developer/AssetTools/Private/AssetTypeActions/AssetTypeActions_Skeleton.cpp
+++ b/Engine/Source/Developer/AssetTools/Private/AssetTypeActions/AssetTypeActions_Skeleton.cpp
@@ -456,6 +456,20 @@ void FAssetTypeActions_Skeleton::RetargetSkeleton(TArray<FAssetToRemapSkeleton>&
 			{
 				if ( UAnimationAsset * AnimAsset = Cast<UAnimationAsset>(Asset) )
 				{
+					UAnimSequenceBase * SequenceBase = Cast<UAnimSequenceBase>(AnimAsset);
+					if (SequenceBase)
+					{
+						// Copy curve data from source asset, preserving data in the target if present.
+						FSmartNameMapping* OldNameMapping = OldSkeleton->SmartNames.GetContainer(USkeleton::AnimCurveMappingName);
+						FSmartNameMapping* NewNameMapping = NewSkeleton->SmartNames.GetContainer(USkeleton::AnimCurveMappingName);
+						SequenceBase->RawCurveData.UpdateLastObservedNames(OldNameMapping);
+
+						for(FFloatCurve& Curve : SequenceBase->RawCurveData.FloatCurves)
+						{
+							NewNameMapping->AddName(Curve.LastObservedName, Curve.CurveUid);
+						}
+					}
+					
 					AnimAsset->ReplaceSkeleton(NewSkeleton, bConvertSpaces);
 				}
 				else if ( UAnimBlueprint * AnimBlueprint = Cast<UAnimBlueprint>(Asset) )
