Fixed Clang compile error

diff --git a/Engine/Source/Developer/AssetTools/Private/AssetTypeActions/AssetTypeActions_Skeleton.cpp b/Engine/Source/Developer/AssetTools/Private/AssetTypeActions/AssetTypeActions_Skeleton.cpp
--- a/Engine/Source/Developer/AssetTools/Private/AssetTypeActions/AssetTypeActions_Skeleton.cpp
+++ b/Engine/Source/Developer/AssetTools/Private/AssetTypeActions/AssetTypeActions_Skeleton.cpp
@@ -567,7 +567,7 @@ void FAssetTypeActions_Skeleton::CreateRig(const TWeakObjectPtr<USkeleton> Skele
 	if(Skeleton.IsValid())
 	{
 		FCreateRigDlg CreateRigDlg(Skeleton.Get());
-		if(CreateRigDlg.ShowModal() == FDlgMergeSkeleton::Confirm)
+		if(CreateRigDlg.ShowModal() == FCreateRigDlg::Confirm)
 		{
 			check(CreateRigDlg.RequiredBones.Num() > 0);
 
