- Fixed the window to function again without being modal
- Fixed issue with when non-existing node convert between

diff --git a/Engine/Source/Developer/AssetTools/Private/AssetTypeActions/AssetTypeActions_AnimBlueprint.cpp b/Engine/Source/Developer/AssetTools/Private/AssetTypeActions/AssetTypeActions_AnimBlueprint.cpp
--- a/Engine/Source/Developer/AssetTools/Private/AssetTypeActions/AssetTypeActions_AnimBlueprint.cpp
+++ b/Engine/Source/Developer/AssetTools/Private/AssetTypeActions/AssetTypeActions_AnimBlueprint.cpp
@@ -11,7 +11,7 @@ void FAssetTypeActions_AnimBlueprint::GetActions( const TArray<UObject*>& InObje
 {
 	FAssetTypeActions_Blueprint::GetActions(InObjects, MenuBuilder);
 
-	auto AnimBlueprints = GetTypedWeakObjectPtrs<UAnimBlueprint>(InObjects);
+	auto AnimBlueprints = FObjectEditorUtils::GetTypedWeakObjectPtrs<UAnimBlueprint>(InObjects);
 
 	MenuBuilder.AddMenuEntry(
 		LOCTEXT("AnimBlueprint_FindSkeleton", "Find Skeleton"),
@@ -147,10 +147,14 @@ void FAssetTypeActions_AnimBlueprint::ExecuteFindSkeleton(TArray<TWeakObjectPtr<
 	}
 }
 
-void FAssetTypeActions_AnimBlueprint::RetargetAssets(const TArray<UObject*> InAnimBlueprints, bool bDuplicateAssets)
+void FAssetTypeActions_AnimBlueprint::RetargetAnimationHandler(USkeleton* OldSkeleton, USkeleton* NewSkeleton, bool bRemapReferencedAssets, bool bConvertSpaces, bool bDuplicateAssets, TArray<TWeakObjectPtr<UObject>> AnimBlueprints)
+{
+	EditorAnimUtils::RetargetAnimations(OldSkeleton, NewSkeleton, AnimBlueprints, bRemapReferencedAssets, bDuplicateAssets, bConvertSpaces);
+}
+
+void FAssetTypeActions_AnimBlueprint::RetargetAssets(TArray<UObject*> InAnimBlueprints, bool bDuplicateAssets)
 {
 	bool bRemapReferencedAssets = false;
-	USkeleton * NewSkeleton = NULL;
 	USkeleton * OldSkeleton = NULL;
 
 	if ( InAnimBlueprints.Num() > 0 )
@@ -159,15 +163,10 @@ void FAssetTypeActions_AnimBlueprint::RetargetAssets(const TArray<UObject*> InAn
 		OldSkeleton = AnimBP->TargetSkeleton;
 	}
 
-	bool bConvertSpaces = OldSkeleton != NULL;
-	bool bShowOnlyCompatibleSkeletons = OldSkeleton != NULL;
-
 	const FText Message = LOCTEXT("RemapSkeleton_Warning", "Select the skeleton to remap this asset to.");
+	auto AnimBlueprints = GetTypedWeakObjectPtrs<UObject>(InAnimBlueprints);
 
-	if (SAnimationRemapSkeleton::ShowModal(OldSkeleton, NewSkeleton, Message, OldSkeleton? &bShowOnlyCompatibleSkeletons : NULL, OldSkeleton ? &bConvertSpaces : NULL, &bRemapReferencedAssets))
-	{
-		EditorAnimUtils::RetargetAnimations(OldSkeleton, NewSkeleton, InAnimBlueprints, bRemapReferencedAssets, bDuplicateAssets, bConvertSpaces);
-	}
+	SAnimationRemapSkeleton::ShowWindow(OldSkeleton, Message, FOnRetargetAnimation::CreateSP(this, &FAssetTypeActions_AnimBlueprint::RetargetAnimationHandler, bDuplicateAssets, AnimBlueprints));
 }
 
 #undef LOCTEXT_NAMESPACE
\ No newline at end of file
diff --git a/Engine/Source/Developer/AssetTools/Private/AssetTypeActions/AssetTypeActions_AnimBlueprint.h b/Engine/Source/Developer/AssetTools/Private/AssetTypeActions/AssetTypeActions_AnimBlueprint.h
--- a/Engine/Source/Developer/AssetTools/Private/AssetTypeActions/AssetTypeActions_AnimBlueprint.h
+++ b/Engine/Source/Developer/AssetTools/Private/AssetTypeActions/AssetTypeActions_AnimBlueprint.h
@@ -23,5 +23,8 @@ private:
 	void ExecuteFindSkeleton(TArray<TWeakObjectPtr<UAnimBlueprint>> Objects);
 
 	/** Context menu item handler for changing the supplied assets skeletons */ 
-	void RetargetAssets(const TArray<UObject*> InAnimBlueprints, bool bDuplicateAssets);
+	void RetargetAssets(TArray<UObject*> InAnimBlueprints, bool bDuplicateAssets);
+
+	/** Handler for retargeting */
+	void RetargetAnimationHandler(USkeleton* OldSkeleton, USkeleton* NewSkeleton, bool bRemapReferencedAssets, bool bConvertSpaces, bool bDuplicateAssets, TArray<TWeakObjectPtr<UObject>> AnimBlueprints);
 };
\ No newline at end of file
diff --git a/Engine/Source/Developer/AssetTools/Private/AssetTypeActions/AssetTypeActions_AnimationAsset.cpp b/Engine/Source/Developer/AssetTools/Private/AssetTypeActions/AssetTypeActions_AnimationAsset.cpp
--- a/Engine/Source/Developer/AssetTools/Private/AssetTypeActions/AssetTypeActions_AnimationAsset.cpp
+++ b/Engine/Source/Developer/AssetTools/Private/AssetTypeActions/AssetTypeActions_AnimationAsset.cpp
@@ -5,12 +5,13 @@
 #include "PersonaModule.h"
 #include "EditorAnimUtils.h"
 #include "SSkeletonWidget.h"
+#include "ObjectEditorUtils.h"
 
 #define LOCTEXT_NAMESPACE "AssetTypeActions"
 
 void FAssetTypeActions_AnimationAsset::GetActions( const TArray<UObject*>& InObjects, FMenuBuilder& MenuBuilder )
 {
-	auto AnimAssets = GetTypedWeakObjectPtrs<UAnimationAsset>(InObjects);
+	auto AnimAssets = FObjectEditorUtils::GetTypedWeakObjectPtrs<UAnimationAsset>(InObjects);
 
 	MenuBuilder.AddMenuEntry(
 		LOCTEXT("AnimSequenceBase_FindSkeleton", "Find Skeleton"),
@@ -153,7 +154,12 @@ void FAssetTypeActions_AnimationAsset::ExecuteFindSkeleton(TArray<TWeakObjectPtr
 	}
 }
 
-void FAssetTypeActions_AnimationAsset::RetargetAssets(const TArray<UObject*> InAnimAssets, bool bDuplicateAssets)
+void FAssetTypeActions_AnimationAsset::RetargetAnimationHandler(USkeleton* OldSkeleton, USkeleton* NewSkeleton, bool bRemapReferencedAssets, bool bConvertSpaces, bool bDuplicateAssets,  TArray<TWeakObjectPtr<UObject>> InAnimAssets)
+{
+	EditorAnimUtils::RetargetAnimations(OldSkeleton, NewSkeleton, InAnimAssets, bRemapReferencedAssets, bDuplicateAssets, bConvertSpaces);
+}
+
+void FAssetTypeActions_AnimationAsset::RetargetAssets(TArray<UObject*> InAnimAssets, bool bDuplicateAssets)
 {
 	bool bRemapReferencedAssets = false;
 	USkeleton * NewSkeleton = NULL;
@@ -164,15 +170,10 @@ void FAssetTypeActions_AnimationAsset::RetargetAssets(const TArray<UObject*> InA
 		OldSkeleton = AnimAsset->GetSkeleton();
 	}
 
-	bool bShowOnlyCompatibleSkeletons = OldSkeleton != NULL;
-	bool bConvertSpaces = OldSkeleton != NULL;
-
 	const FText Message = LOCTEXT("SelectSkeletonToRemap", "Select the skeleton to remap this asset to.");
 
-	if (SAnimationRemapSkeleton::ShowModal(OldSkeleton, NewSkeleton, Message, OldSkeleton? &bShowOnlyCompatibleSkeletons : NULL, OldSkeleton ? &bConvertSpaces : NULL, &bRemapReferencedAssets))
-	{
-		EditorAnimUtils::RetargetAnimations(OldSkeleton, NewSkeleton, InAnimAssets, bRemapReferencedAssets, bDuplicateAssets, bConvertSpaces);
-	}
+	auto AnimAssets = GetTypedWeakObjectPtrs<UObject>(InAnimAssets);
+	SAnimationRemapSkeleton::ShowWindow(OldSkeleton, Message, FOnRetargetAnimation::CreateSP(this, &FAssetTypeActions_AnimationAsset::RetargetAnimationHandler, bDuplicateAssets, AnimAssets) );
 }
 
 
diff --git a/Engine/Source/Developer/AssetTools/Private/AssetTypeActions/AssetTypeActions_AnimationAsset.h b/Engine/Source/Developer/AssetTools/Private/AssetTypeActions/AssetTypeActions_AnimationAsset.h
--- a/Engine/Source/Developer/AssetTools/Private/AssetTypeActions/AssetTypeActions_AnimationAsset.h
+++ b/Engine/Source/Developer/AssetTools/Private/AssetTypeActions/AssetTypeActions_AnimationAsset.h
@@ -24,5 +24,8 @@ private:
 	void ExecuteFindSkeleton(TArray<TWeakObjectPtr<UAnimationAsset>> Objects);
 
 	/** Context menu item handler for changing the supplied assets skeletons */ 
-	void RetargetAssets(const TArray<UObject*> InAnimBlueprints, bool bDuplicateAssets);
+	void RetargetAssets(TArray<UObject*> InAnimBlueprints, bool bDuplicateAssets);
+
+	/** Handler for retargeting */
+	void RetargetAnimationHandler(USkeleton* OldSkeleton, USkeleton* NewSkeleton, bool bRemapReferencedAssets, bool bConvertSpaces, bool bDuplicateAssets, TArray<TWeakObjectPtr<UObject>> InAnimAssets);
 };
\ No newline at end of file
diff --git a/Engine/Source/Developer/AssetTools/Private/AssetTypeActions/AssetTypeActions_Skeleton.cpp b/Engine/Source/Developer/AssetTools/Private/AssetTypeActions/AssetTypeActions_Skeleton.cpp
--- a/Engine/Source/Developer/AssetTools/Private/AssetTypeActions/AssetTypeActions_Skeleton.cpp
+++ b/Engine/Source/Developer/AssetTools/Private/AssetTypeActions/AssetTypeActions_Skeleton.cpp
@@ -417,28 +417,20 @@ void FAssetTypeActions_Skeleton::GetActions( const TArray<UObject*>& InObjects,
 			LOCTEXT("CreateSkeletonSubmenu_ToolTip", "Create assets for this skeleton"),
 			FNewMenuDelegate::CreateSP(this, &FAssetTypeActions_Skeleton::FillCreateMenu, Skeletons));
 
-	MenuBuilder.AddMenuEntry(
-		LOCTEXT("Skeleton_Retarget", "Retarget to Another Skeleton"),
-		LOCTEXT("Skeleton_RetargetTooltip", "Allow all animation assets for this skeleton retarget to another skeleton."),
-		FSlateIcon(),
-		FUIAction(
-			FExecuteAction::CreateSP( this, &FAssetTypeActions_Skeleton::ExecuteRetargetSkeleton, Skeletons ),
-			FCanExecuteAction()
-			)
-		);
-
-	if (InObjects.Num() == 1)
+	// only show if one is selected. It won't work since I changed the window to be normal window
+	if (Skeletons.Num() == 1)
 	{
 		MenuBuilder.AddMenuEntry(
-			LOCTEXT("Skeleton_CreateRig", "Create Rig"),
-			LOCTEXT("Skeleton_CreateRigTooltip", "Create Rig from this skeleton."),
+			LOCTEXT("Skeleton_Retarget", "Retarget to Another Skeleton"),
+			LOCTEXT("Skeleton_RetargetTooltip", "Allow all animation assets for this skeleton retarget to another skeleton."),
 			FSlateIcon(),
 			FUIAction(
-				FExecuteAction::CreateSP(this, &FAssetTypeActions_Skeleton::ExecuteCreateRig, Skeletons),
+				FExecuteAction::CreateSP(this, &FAssetTypeActions_Skeleton::ExecuteRetargetSkeleton, Skeletons),
 				FCanExecuteAction()
 				)
 			);
 	}
+
 	// @todo ImportAnimation
 	/*
 	MenuBuilder.AddMenuEntry(
@@ -468,6 +460,22 @@ void FAssetTypeActions_Skeleton::FillCreateMenu(FMenuBuilder& MenuBuilder, TArra
 	FAnimAssetCreated OnAssetCreated;
 	OnAssetCreated.BindSP(this, &FAssetTypeActions_Skeleton::OnAssetCreated);
 	AnimationEditorUtils::FillCreateAssetMenu(MenuBuilder, Skeletons, /*OnAssetCreated*/FAnimAssetCreated::CreateSP(this, &FAssetTypeActions_Skeleton::OnAssetCreated));
+
+	MenuBuilder.AddMenuSeparator();
+
+	// create rig
+	if(Skeletons.Num() == 1)
+	{
+		MenuBuilder.AddMenuEntry(
+			LOCTEXT("Skeleton_CreateRig", "Create Rig"),
+			LOCTEXT("Skeleton_CreateRigTooltip", "Create Rig from this skeleton."),
+			FSlateIcon(),
+			FUIAction(
+				FExecuteAction::CreateSP(this, &FAssetTypeActions_Skeleton::ExecuteCreateRig, Skeletons),
+				FCanExecuteAction()
+				)
+			);
+	}
 }
 
 void FAssetTypeActions_Skeleton::OpenAssetEditor( const TArray<UObject*>& InObjects, TSharedPtr<IToolkitHost> EditWithinLevelEditor )
@@ -601,44 +609,41 @@ void FAssetTypeActions_Skeleton::CreateAnimationAssets(const TArray<TWeakObjectP
 	AnimationEditorUtils::CreateAnimationAssets(Skeletons, AssetClass, InPrefix, /*OnAssetCreated*/FAnimAssetCreated::CreateSP(this, &FAssetTypeActions_Skeleton::OnAssetCreated));
 }
 
+void FAssetTypeActions_Skeleton::RetargetAnimationHandler(USkeleton* OldSkeleton, USkeleton* NewSkeleton, bool bRemapReferencedAssets, bool bConvertSpaces)
+{
+	// find all assets who references old skeleton
+	TArray<FName> Packages;
+
+	// If the asset registry is still loading assets, we cant check for referencers, so we must open the rename dialog
+	FAssetRegistryModule& AssetRegistryModule = FModuleManager::LoadModuleChecked<FAssetRegistryModule>("AssetRegistry");
+	AssetRegistryModule.Get().GetReferencers(OldSkeleton->GetOutermost()->GetFName(), Packages);
+
+	if(AssetRegistryModule.Get().IsLoadingAssets())
+	{
+		// Open a dialog asking the user to wait while assets are being discovered
+		SDiscoveringAssetsDialog::OpenDiscoveringAssetsDialog(
+			SDiscoveringAssetsDialog::FOnAssetsDiscovered::CreateSP(this, &FAssetTypeActions_Skeleton::PerformRetarget, OldSkeleton, NewSkeleton, Packages, bConvertSpaces)
+			);
+	}
+	else
+	{
+		PerformRetarget(OldSkeleton, NewSkeleton, Packages, bConvertSpaces);
+	}
+}
 void FAssetTypeActions_Skeleton::ExecuteRetargetSkeleton(TArray<TWeakObjectPtr<USkeleton>> Skeletons)
 {
+	// only allow 1 for now, it is scary to do this for multi
 	// warn the user to shut down any persona that is opened
 	if ( FMessageDialog::Open(EAppMsgType::YesNo, LOCTEXT("CloseReferencingEditors", "You need to close Persona or anything that references animation, mesh or animation blueprint before this step. Continue?")) == EAppReturnType::Yes )
 	{
 		for (auto SkelIt = Skeletons.CreateConstIterator(); SkelIt; ++SkelIt)
 		{	
 			USkeleton * OldSkeleton = (*SkelIt).Get();
-			USkeleton * NewSkeleton = NULL;
 
 			const FText Message = LOCTEXT("RetargetSkeleton_Warning", "This only converts animation data -i.e. animation assets and Anim Blueprints. \nIf you'd like to convert SkeletalMesh, use the context menu (Assign Skeleton) for each mesh. \n\nIf you'd like to convert mesh as well, please do so before converting animation data. \nOtherwise you will lose any extra track that is in the new mesh.");
 
-			bool bConvertSpaces = OldSkeleton != NULL;
-			bool bShowOnlyCompatibleSkeletons = OldSkeleton != NULL;
-
 			// ask user what they'd like to change to 
-			if (SAnimationRemapSkeleton::ShowModal(OldSkeleton, NewSkeleton, Message, OldSkeleton? &bShowOnlyCompatibleSkeletons : NULL, OldSkeleton ? &bConvertSpaces : NULL))
-			{
-				// find all assets who references old skeleton
-				TArray<FName> Packages;
-
-				// If the asset registry is still loading assets, we cant check for referencers, so we must open the rename dialog
-				FAssetRegistryModule& AssetRegistryModule = FModuleManager::LoadModuleChecked<FAssetRegistryModule>("AssetRegistry");
-				AssetRegistryModule.Get().GetReferencers(OldSkeleton->GetOutermost()->GetFName(), Packages);
-
-				if ( AssetRegistryModule.Get().IsLoadingAssets() )
-				{
-					// Open a dialog asking the user to wait while assets are being discovered
-					SDiscoveringAssetsDialog::OpenDiscoveringAssetsDialog(
-						SDiscoveringAssetsDialog::FOnAssetsDiscovered::CreateSP(this, &FAssetTypeActions_Skeleton::PerformRetarget, OldSkeleton, NewSkeleton, Packages, bConvertSpaces)
-						);
-				}
-				else
-				{
-					PerformRetarget(OldSkeleton, NewSkeleton, Packages, bConvertSpaces);
-				}
-				
-			}
+			SAnimationRemapSkeleton::ShowWindow(OldSkeleton, Message, FOnRetargetAnimation::CreateSP(this, &FAssetTypeActions_Skeleton::RetargetAnimationHandler));
 		}
 	}
 }
diff --git a/Engine/Source/Developer/AssetTools/Private/AssetTypeActions/AssetTypeActions_Skeleton.h b/Engine/Source/Developer/AssetTools/Private/AssetTypeActions/AssetTypeActions_Skeleton.h
--- a/Engine/Source/Developer/AssetTools/Private/AssetTypeActions/AssetTypeActions_Skeleton.h
+++ b/Engine/Source/Developer/AssetTools/Private/AssetTypeActions/AssetTypeActions_Skeleton.h
@@ -77,4 +77,7 @@ private: // Helper functions
 	void DetectReadOnlyPackages(TArray<FAssetToRemapSkeleton>& AssetsToRemap, TArray<UPackage*>& InOutPackagesToSave) const;
 	void FillCreateMenu(FMenuBuilder& MenuBuilder, TArray<TWeakObjectPtr<USkeleton>> Skeletons) const;
 	void OnAssetCreated(TArray<UObject*> NewAssets) const;
+
+	/** Handler for retargeting */
+	void RetargetAnimationHandler(USkeleton* OldSkeleton, USkeleton* NewSkeleton, bool bRemapReferencedAssets, bool bConvertSpaces);
 };
\ No newline at end of file
diff --git a/Engine/Source/Editor/Persona/Private/SAnimationSequenceBrowser.cpp b/Engine/Source/Editor/Persona/Private/SAnimationSequenceBrowser.cpp
--- a/Engine/Source/Editor/Persona/Private/SAnimationSequenceBrowser.cpp
+++ b/Engine/Source/Editor/Persona/Private/SAnimationSequenceBrowser.cpp
@@ -15,6 +15,7 @@
 #include "Runtime/AssetRegistry/Public/AssetRegistryModule.h"
 #include "SceneViewport.h"
 #include "AnimPreviewInstance.h"
+#include "ObjectEditorUtils.h"
 
 #define LOCTEXT_NAMESPACE "SequenceBrowser"
 
@@ -308,6 +309,20 @@ void SAnimationSequenceBrowser::OnReimportAnimation(TArray<FAssetData> SelectedA
 	}
 }
 
+void SAnimationSequenceBrowser::RetargetAnimationHandler(USkeleton* OldSkeleton, USkeleton* NewSkeleton, bool bRemapReferencedAssets, bool bConvertSpaces, TArray<TWeakObjectPtr<UObject>> InAnimAssets)
+{
+	UObject* AssetToOpen = EditorAnimUtils::RetargetAnimations(OldSkeleton, NewSkeleton, InAnimAssets, bRemapReferencedAssets, true, bConvertSpaces);
+
+	if(UAnimationAsset* AnimAsset = Cast<UAnimationAsset>(AssetToOpen))
+	{
+		FAssetRegistryModule::AssetCreated(AssetToOpen);
+		// once all success, attempt to open new persona module with new skeleton
+		EToolkitMode::Type Mode = EToolkitMode::Standalone;
+		FPersonaModule& PersonaModule = FModuleManager::LoadModuleChecked<FPersonaModule>("Persona");
+		PersonaModule.CreatePersona(Mode, TSharedPtr<IToolkitHost>(), NewSkeleton, NULL, AnimAsset, NULL);
+	}
+}
+
 void SAnimationSequenceBrowser::OnCreateCopy(TArray<FAssetData> Selected)
 {
 	if ( Selected.Num() > 0 )
@@ -315,27 +330,26 @@ void SAnimationSequenceBrowser::OnCreateCopy(TArray<FAssetData> Selected)
 		// ask which skeleton users would like to choose
 		USkeleton * OldSkeleton = PersonaPtr.Pin()->GetSkeleton();
 		USkeleton * NewSkeleton = NULL;
-		bool		bRemapReferencedAssets = true;
-		bool		bConvertSpaces = true;
-		bool		bShowOnlyCompatibleSkeleton = true;
 		bool		bDuplicateAssets = true;
 
 		const FText Message = LOCTEXT("RemapSkeleton_Warning", "This will duplicate the asset and convert to new skeleton.");
 
-		// ask user what they'd like to change to 
-		if (SAnimationRemapSkeleton::ShowModal(OldSkeleton, NewSkeleton, Message, &bShowOnlyCompatibleSkeleton, &bConvertSpaces, &bRemapReferencedAssets))
+		TArray<UObject *> AnimAssets;
+		for ( auto SelectedAsset : Selected )
 		{
-			UObject* AssetToOpen = EditorAnimUtils::RetargetAnimations(OldSkeleton, NewSkeleton, Selected, bRemapReferencedAssets, bDuplicateAssets, bConvertSpaces);
-
-			if(UAnimationAsset* AnimAsset = Cast<UAnimationAsset>(AssetToOpen))
+			UAnimationAsset* Asset = Cast<UAnimationAsset>(SelectedAsset.GetAsset());
+			if (Asset)
 			{
-				FAssetRegistryModule::AssetCreated(AssetToOpen);
-				// once all success, attempt to open new persona module with new skeleton
-				EToolkitMode::Type Mode = EToolkitMode::Standalone;
-				FPersonaModule& PersonaModule = FModuleManager::LoadModuleChecked<FPersonaModule>( "Persona" );
-				PersonaModule.CreatePersona( Mode, TSharedPtr<IToolkitHost>(), NewSkeleton, NULL, AnimAsset, NULL );
+				AnimAssets.Add(Asset);
 			}
 		}
+
+		if (AnimAssets.Num() > 0)
+		{
+			auto AnimAssetsToConvert = FObjectEditorUtils::GetTypedWeakObjectPtrs<UObject>(AnimAssets);
+			// ask user what they'd like to change to 
+			SAnimationRemapSkeleton::ShowWindow(OldSkeleton, Message, FOnRetargetAnimation::CreateSP(this, &SAnimationSequenceBrowser::RetargetAnimationHandler, AnimAssetsToConvert));
+		}
 	}
 }
 
diff --git a/Engine/Source/Editor/Persona/Private/SAnimationSequenceBrowser.h b/Engine/Source/Editor/Persona/Private/SAnimationSequenceBrowser.h
--- a/Engine/Source/Editor/Persona/Private/SAnimationSequenceBrowser.h
+++ b/Engine/Source/Editor/Persona/Private/SAnimationSequenceBrowser.h
@@ -188,4 +188,6 @@ protected:
 	// delegate to sync the asset picker to selected assets
 	FSyncToAssetsDelegate SyncToAssetsDelegate;
 	FGetCurrentSelectionDelegate GetCurrentSelectionDelegate;
+
+	void RetargetAnimationHandler(USkeleton* OldSkeleton, USkeleton* NewSkeleton, bool bRemapReferencedAssets, bool bConvertSpaces, TArray<TWeakObjectPtr<UObject>> InAnimAssets);
 };
\ No newline at end of file
diff --git a/Engine/Source/Editor/UnrealEd/Private/EditorAnimUtils.cpp b/Engine/Source/Editor/UnrealEd/Private/EditorAnimUtils.cpp
--- a/Engine/Source/Editor/UnrealEd/Private/EditorAnimUtils.cpp
+++ b/Engine/Source/Editor/UnrealEd/Private/EditorAnimUtils.cpp
@@ -10,7 +10,7 @@
 #include "Developer/AssetTools/Public/AssetToolsModule.h"
 #include "NotificationManager.h"
 #include "Editor/Persona/Public/PersonaModule.h"
-
+#include "ObjectEditorUtils.h"
 
 #define LOCTEXT_NAMESPACE "EditorAnimUtils"
 
@@ -18,7 +18,6 @@ namespace EditorAnimUtils
 {
 	//////////////////////////////////////////////////////////////////
 	// FAnimationRetargetContext
-
 	FAnimationRetargetContext::FAnimationRetargetContext(const TArray<FAssetData>& AssetsToRetarget, bool bRetargetReferredAssets, bool bInConvertAnimationDataInComponentSpaces) 
 		: SingleTargetObject(NULL)
 		, bConvertAnimationDataInComponentSpaces(bInConvertAnimationDataInComponentSpaces)
@@ -28,21 +27,22 @@ namespace EditorAnimUtils
 		{
 			Objects.Add((*Iter).GetAsset());
 		}
-		Initialize(Objects,bRetargetReferredAssets);
+		auto WeakObjectList = FObjectEditorUtils::GetTypedWeakObjectPtrs<UObject>(Objects);
+		Initialize(WeakObjectList,bRetargetReferredAssets);
 	}
 
-	FAnimationRetargetContext::FAnimationRetargetContext(const TArray<UObject*>& AssetsToRetarget, bool bRetargetReferredAssets, bool bInConvertAnimationDataInComponentSpaces) 
+	FAnimationRetargetContext::FAnimationRetargetContext(TArray<TWeakObjectPtr<UObject>> AssetsToRetarget, bool bRetargetReferredAssets, bool bInConvertAnimationDataInComponentSpaces) 
 		: SingleTargetObject(NULL)
 		, bConvertAnimationDataInComponentSpaces(bInConvertAnimationDataInComponentSpaces)
 	{
 		Initialize(AssetsToRetarget,bRetargetReferredAssets);
 	}
 
-	void FAnimationRetargetContext::Initialize(const TArray<UObject*>& AssetsToRetarget, bool bRetargetReferredAssets)
+	void FAnimationRetargetContext::Initialize(TArray<TWeakObjectPtr<UObject>> AssetsToRetarget, bool bRetargetReferredAssets)
 	{
 		for(auto Iter = AssetsToRetarget.CreateConstIterator(); Iter; ++Iter)
 		{
-			UObject* Asset = (*Iter);
+			UObject* Asset = (*Iter).Get();
 			if( UAnimSequence* AnimSeq = Cast<UAnimSequence>(Asset) )
 			{
 				AnimSequencesToRetarget.AddUnique(AnimSeq);
@@ -60,7 +60,7 @@ namespace EditorAnimUtils
 		if(AssetsToRetarget.Num() == 1)
 		{
 			//Only chose one object to retarget, keep track of it
-			SingleTargetObject = AssetsToRetarget[0];
+			SingleTargetObject = AssetsToRetarget[0].Get();
 		}
 
 		if(bRetargetReferredAssets)
@@ -224,7 +224,7 @@ namespace EditorAnimUtils
 	}
 
 	//////////////////////////////////////////////////////////////////
-	UObject* RetargetAnimations(USkeleton * OldSkeleton, USkeleton* NewSkeleton, const TArray<UObject*>& AssetsToRetarget, bool bRetargetReferredAssets, bool bDuplicateAssetsBeforeRetarget, bool bConvertSpace)
+	UObject* RetargetAnimations(USkeleton * OldSkeleton, USkeleton* NewSkeleton, TArray<TWeakObjectPtr<UObject>> AssetsToRetarget, bool bRetargetReferredAssets, bool bDuplicateAssetsBeforeRetarget, bool bConvertSpace)
 	{
 		FAnimationRetargetContext RetargetContext(AssetsToRetarget, bRetargetReferredAssets, bConvertSpace);
 		return RetargetAnimations(OldSkeleton, NewSkeleton, RetargetContext, bRetargetReferredAssets, bDuplicateAssetsBeforeRetarget);
diff --git a/Engine/Source/Editor/UnrealEd/Private/SSkeletonWidget.cpp b/Engine/Source/Editor/UnrealEd/Private/SSkeletonWidget.cpp
--- a/Engine/Source/Editor/UnrealEd/Private/SSkeletonWidget.cpp
+++ b/Engine/Source/Editor/UnrealEd/Private/SSkeletonWidget.cpp
@@ -400,6 +400,7 @@ void SAnimationRemapSkeleton::Construct( const FArguments& InArgs )
 	bRemapReferencedAssets = true;
 	bConvertSpaces = false;
 	bShowOnlyCompatibleSkeletons = false;
+	OnRetargetAnimationDelegate = InArgs._OnRetargetDelegate;
 
 	TSharedRef<SVerticalBox> Widget = SNew(SVerticalBox);
 	if(InArgs._ShowRemapOption)
@@ -681,6 +682,11 @@ void SAnimationRemapSkeleton::OnAssetSelectedFromPicker(const FAssetData& AssetD
 
 FReply SAnimationRemapSkeleton::OnApply()
 {
+	if (OnRetargetAnimationDelegate.IsBound())
+	{
+		OnRetargetAnimationDelegate.Execute(OldSkeleton, NewSkeleton, bRemapReferencedAssets, bConvertSpaces);
+	}
+
 	CloseWindow();
 	return FReply::Handled();
 }
@@ -699,7 +705,7 @@ void SAnimationRemapSkeleton::CloseWindow()
 	}
 }
 
-bool SAnimationRemapSkeleton::ShowModal(USkeleton * OldSkeleton, USkeleton * & NewSkeleton, const FText& WarningMessage, bool *bShowOnlyCompatibleSkeletons, bool * bConvertSpace, bool * bRemapReferencedAssets)
+void SAnimationRemapSkeleton::ShowWindow(USkeleton * OldSkeleton, const FText& WarningMessage, FOnRetargetAnimation RetargetDelegate)
 {
 	static TSharedPtr<SWindow> DialogWindow;
 
@@ -724,28 +730,15 @@ bool SAnimationRemapSkeleton::ShowModal(USkeleton * OldSkeleton, USkeleton * & N
 			.CurrentSkeleton(OldSkeleton)
 			.WidgetWindow(DialogWindow)
 			.WarningMessage(WarningMessage)
-			.ShowRemapOption(bRemapReferencedAssets != NULL)
-			.ShowConvertSpacesOption(bConvertSpace != NULL)
-			.ShowCompatibleDisplayOption(bShowOnlyCompatibleSkeletons != NULL)
+			.ShowRemapOption(true)
+			.ShowConvertSpacesOption(OldSkeleton != NULL)
+			.ShowCompatibleDisplayOption(OldSkeleton != NULL)
+			.OnRetargetDelegate(RetargetDelegate)
 		];
 
 	DialogWindow->SetContent(DialogWrapper.ToSharedRef());
 
 	FSlateApplication::Get().AddWindow(DialogWindow.ToSharedRef());
-
-	NewSkeleton = DialogWidget.Get()->NewSkeleton;
-
-	if(bRemapReferencedAssets)
-	{
-		*bRemapReferencedAssets = DialogWidget.Get()->bRemapReferencedAssets;
-	}
-
-	if (bConvertSpace)
-	{
-		*bConvertSpace = DialogWidget.Get()->bConvertSpaces;
-	}
-
-	return (NewSkeleton != NULL && NewSkeleton != OldSkeleton);
 }
 
 ////////////////////////////////////////////////////
diff --git a/Engine/Source/Editor/UnrealEd/Public/EditorAnimUtils.h b/Engine/Source/Editor/UnrealEd/Public/EditorAnimUtils.h
--- a/Engine/Source/Editor/UnrealEd/Public/EditorAnimUtils.h
+++ b/Engine/Source/Editor/UnrealEd/Public/EditorAnimUtils.h
@@ -8,7 +8,7 @@ namespace EditorAnimUtils
 	{
 	public:
 		FAnimationRetargetContext(const TArray<FAssetData>& AssetsToRetarget, bool bRetargetReferredAssets, bool bInConvertAnimationDataInComponentSpaces);
-		FAnimationRetargetContext(const TArray<UObject*>& AssetsToRetarget, bool bRetargetReferredAssets, bool bInConvertAnimationDataInComponentSpaces);
+		FAnimationRetargetContext(TArray<TWeakObjectPtr<UObject>> AssetsToRetarget, bool bRetargetReferredAssets, bool bInConvertAnimationDataInComponentSpaces);
 
 		/** Were we supplied anything that we can retarget */
 		bool HasAssetsToRetarget() const;
@@ -46,7 +46,7 @@ namespace EditorAnimUtils
 		bool bConvertAnimationDataInComponentSpaces;
 
 		/** Initialize the object, only to be called by constructors */
-		void Initialize(const TArray<UObject*>& AssetsToRetarget, bool bRetargetReferredAssets);
+		void Initialize(TArray<TWeakObjectPtr<UObject>> AssetsToRetarget, bool bRetargetReferredAssets);
 	};
 
 	/**
@@ -58,7 +58,7 @@ namespace EditorAnimUtils
 	 * @param bDuplicatedAssetsBeforeRetarget	If true the assets are duplicated and then the duplicates are retargetted
 	 * @param bConvertSpace						Do the conversion in component space of the animation to match new target
 	 */
-	UNREALED_API UObject* RetargetAnimations(USkeleton * OldSkeleton, USkeleton* NewSkeleton, const TArray<UObject*>& AssetsToRetarget, bool bRetargetReferredAssets, bool bDuplicateAssetsBeforeRetarget, bool bConvertSpace);
+	UNREALED_API UObject* RetargetAnimations(USkeleton * OldSkeleton, USkeleton* NewSkeleton, TArray<TWeakObjectPtr<UObject>> AssetsToRetarget, bool bRetargetReferredAssets, bool bDuplicateAssetsBeforeRetarget, bool bConvertSpace);
 
 	/**
 	 * Retargets the supplied FAssetDatas (as long as they are an animation asset), optionally duplicating them and retargetting their reference assets too
diff --git a/Engine/Source/Editor/UnrealEd/Public/SSkeletonWidget.h b/Engine/Source/Editor/UnrealEd/Public/SSkeletonWidget.h
--- a/Engine/Source/Editor/UnrealEd/Public/SSkeletonWidget.h
+++ b/Engine/Source/Editor/UnrealEd/Public/SSkeletonWidget.h
@@ -291,6 +291,9 @@ private:
 /** 
  * Slate panel for choose Skeleton for assets to relink
  */
+
+DECLARE_DELEGATE_FourParams(FOnRetargetAnimation, USkeleton* /*OldSkeleton*/, USkeleton* /*NewSkeleton*/, bool /*bRemapReferencedAssets*/, bool /*bConvertSpaces*/)
+
 class SAnimationRemapSkeleton : public SCompoundWidget
 {
 public:
@@ -309,7 +312,7 @@ public:
 		SLATE_ARGUMENT( bool, ShowRemapOption )
 		SLATE_ARGUMENT( bool, ShowConvertSpacesOption )
 		SLATE_ARGUMENT( bool, ShowCompatibleDisplayOption )
-
+		SLATE_EVENT(FOnRetargetAnimation, OnRetargetDelegate)
 	SLATE_END_ARGS()	
 
 	/**
@@ -352,6 +355,8 @@ private:
 
 	TWeakPtr<SWindow> WidgetWindow;
 
+	FOnRetargetAnimation OnRetargetAnimationDelegate;
+
 	/** Handlers for check box for remapping assets option */
 	ESlateCheckBoxState::Type IsRemappingReferencedAssets() const;
 	void OnRemappingReferencedAssetsChanged(ESlateCheckBoxState::Type InNewRadioState);
@@ -390,14 +395,13 @@ private:
 public:
 
 	/**
-	 *  Show Modal window
+	 *  Show window
 	 *
 	 * @param OldSkeleton		Old Skeleton to change from
-	 * @param NewSkeleton(out)	New Selected Skeleton
 	 *
 	 * @return true if successfully selected new skeleton
 	 */
-	static UNREALED_API bool ShowModal(USkeleton * OldSkeleton, USkeleton * & NewSkeleton, const FText& WarningMessage, bool *bShowOnlyCompatibleSkeletons, bool * bConvertSpace = NULL, bool * bRemapReferencedAssets=NULL);
+	static UNREALED_API void ShowWindow(USkeleton * OldSkeleton, const FText& WarningMessage, FOnRetargetAnimation RetargetDelegate);
 };
 
 ////////////////////////////////////////////////////
diff --git a/Engine/Source/Runtime/Engine/Classes/Animation/Skeleton.h b/Engine/Source/Runtime/Engine/Classes/Animation/Skeleton.h
--- a/Engine/Source/Runtime/Engine/Classes/Animation/Skeleton.h
+++ b/Engine/Source/Runtime/Engine/Classes/Animation/Skeleton.h
@@ -632,6 +632,8 @@ public:
 	ENGINE_API void SetRigConfig(URig * Rig);
 	ENGINE_API FName GetRigBoneMapping(const FName & NodeName) const;
 	ENGINE_API bool SetRigBoneMapping(const FName & NodeName, FName BoneName);
+	// this make sure it stays within the valid range
+	ENGINE_API int32 GetMappedValidNodes(TArray<FName> &OutValidNodeNames);
 	// verify if it has all latest data
 	ENGINE_API void RefreshRigConfig();
 	int32 FindRigBoneMapping(const FName & NodeName) const;
diff --git a/Engine/Source/Runtime/Engine/Private/Animation/Rig.cpp b/Engine/Source/Runtime/Engine/Private/Animation/Rig.cpp
--- a/Engine/Source/Runtime/Engine/Private/Animation/Rig.cpp
+++ b/Engine/Source/Runtime/Engine/Private/Animation/Rig.cpp
@@ -294,6 +294,7 @@ void URig::CreateFromSkeleton(const USkeleton * Skeleton, const TMap<int32, int3
 	}
 }
 
+/*
 void URig::CalculateComponentSpace(int32 NodeIndex, const FTransform & LocalTransform, const TArray<FTransform> & TransformBuffer, const FGetParentIndex& DelegateToGetParentIndex, FTransform & OutComponentSpaceTransform) const
 {
 	int32 ConstraintIndex = FindRigControlByNodeName(Nodes[NodeIndex].Name);
@@ -409,7 +410,7 @@ void URig::CalculateLocalSpace(int32 NodeIndex, const FTransform & ComponentTran
 		// @todo fix this
 		OutLocalSpaceTransform.SetScale3D(ComponentTransform.GetScale3D());
 	}
-}
+}*/
 
 void URig::SetAllConstraintsToParents()
 {
diff --git a/Engine/Source/Runtime/Engine/Private/Animation/Skeleton.cpp b/Engine/Source/Runtime/Engine/Private/Animation/Skeleton.cpp
--- a/Engine/Source/Runtime/Engine/Private/Animation/Skeleton.cpp
+++ b/Engine/Source/Runtime/Engine/Private/Animation/Skeleton.cpp
@@ -1042,6 +1042,21 @@ FName USkeleton::GetRigBoneMapping(const FName & NodeName) const
 	return NAME_None;
 }
 
+int32 USkeleton::GetMappedValidNodes(TArray<FName> &OutValidNodeNames)
+{
+	OutValidNodeNames.Empty();
+
+	for (auto Entry : RigConfig.BoneMappingTable)
+	{
+		if (Entry.BoneName != NAME_None)
+		{
+			OutValidNodeNames.Add(Entry.NodeName);
+		}
+	}
+
+	return OutValidNodeNames.Num();
+}
+
 bool USkeleton::SetRigBoneMapping(const FName & NodeName, FName BoneName)
 {
 	// make sure it's valid
diff --git a/Engine/Source/Runtime/Engine/Public/ObjectEditorUtils.h b/Engine/Source/Runtime/Engine/Public/ObjectEditorUtils.h
--- a/Engine/Source/Runtime/Engine/Public/ObjectEditorUtils.h
+++ b/Engine/Source/Runtime/Engine/Public/ObjectEditorUtils.h
@@ -123,6 +123,21 @@ namespace FObjectEditorUtils
 
 		return true;
 	}
+
+	/** Helper function to convert the input for GetActions to a list that can be used for delegates */
+	template <typename T>
+	static TArray<TWeakObjectPtr<T>> GetTypedWeakObjectPtrs(const TArray<UObject*>& InObjects)
+	{
+		check(InObjects.Num() > 0);
+
+		TArray<TWeakObjectPtr<T>> TypedObjects;
+		for(auto ObjIt = InObjects.CreateConstIterator(); ObjIt; ++ObjIt)
+		{
+			TypedObjects.Add(CastChecked<T>(*ObjIt));
+		}
+
+		return TypedObjects;
+	}
 };
 
 #endif // WITH_EDITOR
