Reverting wrong change and fixing compile error

diff --git a/Engine/Source/Developer/AssetTools/Private/AssetTypeActions/AssetTypeActions_AnimBlueprint.cpp b/Engine/Source/Developer/AssetTools/Private/AssetTypeActions/AssetTypeActions_AnimBlueprint.cpp
--- a/Engine/Source/Developer/AssetTools/Private/AssetTypeActions/AssetTypeActions_AnimBlueprint.cpp
+++ b/Engine/Source/Developer/AssetTools/Private/AssetTypeActions/AssetTypeActions_AnimBlueprint.cpp
@@ -11,7 +11,7 @@ void FAssetTypeActions_AnimBlueprint::GetActions( const TArray<UObject*>& InObje
 {
 	FAssetTypeActions_Blueprint::GetActions(InObjects, MenuBuilder);
 
-	auto AnimBlueprints = FObjectEditorUtils::GetTypedWeakObjectPtrs<UAnimBlueprint>(InObjects);
+	auto AnimBlueprints = GetTypedWeakObjectPtrs<UAnimBlueprint>(InObjects);
 
 	MenuBuilder.AddMenuEntry(
 		LOCTEXT("AnimBlueprint_FindSkeleton", "Find Skeleton"),
diff --git a/Engine/Source/Developer/AssetTools/Private/AssetTypeActions/AssetTypeActions_AnimationAsset.cpp b/Engine/Source/Developer/AssetTools/Private/AssetTypeActions/AssetTypeActions_AnimationAsset.cpp
--- a/Engine/Source/Developer/AssetTools/Private/AssetTypeActions/AssetTypeActions_AnimationAsset.cpp
+++ b/Engine/Source/Developer/AssetTools/Private/AssetTypeActions/AssetTypeActions_AnimationAsset.cpp
@@ -11,7 +11,7 @@
 
 void FAssetTypeActions_AnimationAsset::GetActions( const TArray<UObject*>& InObjects, FMenuBuilder& MenuBuilder )
 {
-	auto AnimAssets = FObjectEditorUtils::GetTypedWeakObjectPtrs<UAnimationAsset>(InObjects);
+	auto AnimAssets = GetTypedWeakObjectPtrs<UAnimationAsset>(InObjects);
 
 	MenuBuilder.AddMenuEntry(
 		LOCTEXT("AnimSequenceBase_FindSkeleton", "Find Skeleton"),
diff --git a/Engine/Source/Developer/MovieSceneTools/Private/AssetTypeActions_MovieScene.cpp b/Engine/Source/Developer/MovieSceneTools/Private/AssetTypeActions_MovieScene.cpp
--- a/Engine/Source/Developer/MovieSceneTools/Private/AssetTypeActions_MovieScene.cpp
+++ b/Engine/Source/Developer/MovieSceneTools/Private/AssetTypeActions_MovieScene.cpp
@@ -6,6 +6,7 @@
 #include "Toolkits/IToolkit.h"
 #include "ISequencerModule.h"
 #include "MovieScene.h"
+#include "ObjectEditorUtils.h"
 
 #define LOCTEXT_NAMESPACE "AssetTypeActions"
 
