Remove unnecessary file include

diff --git a/Engine/Source/Developer/AssetTools/Private/AssetTypeActions/AssetTypeActions_AnimationAsset.cpp b/Engine/Source/Developer/AssetTools/Private/AssetTypeActions/AssetTypeActions_AnimationAsset.cpp
--- a/Engine/Source/Developer/AssetTools/Private/AssetTypeActions/AssetTypeActions_AnimationAsset.cpp
+++ b/Engine/Source/Developer/AssetTools/Private/AssetTypeActions/AssetTypeActions_AnimationAsset.cpp
@@ -5,7 +5,6 @@
 #include "PersonaModule.h"
 #include "EditorAnimUtils.h"
 #include "SSkeletonWidget.h"
-#include "ObjectEditorUtils.h"
 
 #define LOCTEXT_NAMESPACE "AssetTypeActions"
 
diff --git a/Engine/Source/Developer/MovieSceneTools/Private/AssetTypeActions_MovieScene.cpp b/Engine/Source/Developer/MovieSceneTools/Private/AssetTypeActions_MovieScene.cpp
--- a/Engine/Source/Developer/MovieSceneTools/Private/AssetTypeActions_MovieScene.cpp
+++ b/Engine/Source/Developer/MovieSceneTools/Private/AssetTypeActions_MovieScene.cpp
@@ -6,7 +6,6 @@
 #include "Toolkits/IToolkit.h"
 #include "ISequencerModule.h"
 #include "MovieScene.h"
-#include "ObjectEditorUtils.h"
 
 #define LOCTEXT_NAMESPACE "AssetTypeActions"
 
