Added event when a import fails, with extension information

#ttp 344941 - Record Import failed events when users drag unsupported formats into the content browser

#branch UE4

diff --git a/Engine/Source/Developer/AssetTools/Private/AssetTools.cpp b/Engine/Source/Developer/AssetTools/Private/AssetTools.cpp
--- a/Engine/Source/Developer/AssetTools/Private/AssetTools.cpp
+++ b/Engine/Source/Developer/AssetTools/Private/AssetTools.cpp
@@ -572,6 +572,14 @@ TArray<UObject*> FAssetTools::ImportAssets(const TArray<FString>& Files, const F
 		}
 		else
 		{
+			if(FEngineAnalytics::IsAvailable())
+			{
+				TArray<FAnalyticsEventAttribute> Attribs;
+				Attribs.Add(FAnalyticsEventAttribute(TEXT("FileExtension"), FileExtension));
+
+				FEngineAnalytics::GetProvider().RecordEvent(TEXT("Editor.Usage.ImportFailed"), Attribs);
+			}
+
 			const FText Message = FText::Format( LOCTEXT("ImportFailed_UnknownExtension", "Failed to import '{0}'. Unknown extension '{1}'."), FText::FromString( Filename ), FText::FromString( FileExtension ) );
 			FNotificationInfo Info(Message);
 			Info.ExpireDuration = 3.0f;
