remove destructible mesh from engine header and add it where it's actually used

diff --git a/Engine/Source/Developer/AssetTools/Private/AssetTypeActions/AssetTypeActions_DestructibleMesh.cpp b/Engine/Source/Developer/AssetTools/Private/AssetTypeActions/AssetTypeActions_DestructibleMesh.cpp
--- a/Engine/Source/Developer/AssetTools/Private/AssetTypeActions/AssetTypeActions_DestructibleMesh.cpp
+++ b/Engine/Source/Developer/AssetTools/Private/AssetTypeActions/AssetTypeActions_DestructibleMesh.cpp
@@ -5,6 +5,8 @@
 #include "Editor/DestructibleMeshEditor/Public/DestructibleMeshEditorModule.h"
 #include "Editor/DestructibleMeshEditor/Public/IDestructibleMeshEditor.h"
 
+#include "Engine/DestructibleMesh.h"
+
 #define LOCTEXT_NAMESPACE "AssetTypeActions"
 
 UClass* FAssetTypeActions_DestructibleMesh::GetSupportedClass() const
diff --git a/Engine/Source/Developer/AssetTools/Private/AssetTypeActions/AssetTypeActions_StaticMesh.cpp b/Engine/Source/Developer/AssetTools/Private/AssetTypeActions/AssetTypeActions_StaticMesh.cpp
--- a/Engine/Source/Developer/AssetTools/Private/AssetTypeActions/AssetTypeActions_StaticMesh.cpp
+++ b/Engine/Source/Developer/AssetTools/Private/AssetTypeActions/AssetTypeActions_StaticMesh.cpp
@@ -8,6 +8,7 @@
 #include "Editor/DestructibleMeshEditor/Public/IDestructibleMeshEditor.h"
 
 #include "ApexDestructibleAssetImport.h"
+#include "Engine/DestructibleMesh.h"
 
 #include "FbxMeshUtils.h"
 
diff --git a/Engine/Source/Editor/DestructibleMeshEditor/Private/DestructibleChunkParamsProxy.cpp b/Engine/Source/Editor/DestructibleMeshEditor/Private/DestructibleChunkParamsProxy.cpp
--- a/Engine/Source/Editor/DestructibleMeshEditor/Private/DestructibleChunkParamsProxy.cpp
+++ b/Engine/Source/Editor/DestructibleMeshEditor/Private/DestructibleChunkParamsProxy.cpp
@@ -2,6 +2,7 @@
 
 #include "DestructibleMeshEditorPrivatePCH.h"
 #include "ApexDestructibleAssetImport.h"
+#include "Engine/DestructibleMesh.h"
 
 UDestructibleChunkParamsProxy::UDestructibleChunkParamsProxy(const class FPostConstructInitializeProperties& PCIP)
 	:Super(PCIP)
diff --git a/Engine/Source/Editor/DestructibleMeshEditor/Private/DestructibleMeshEditor.cpp b/Engine/Source/Editor/DestructibleMeshEditor/Private/DestructibleMeshEditor.cpp
--- a/Engine/Source/Editor/DestructibleMeshEditor/Private/DestructibleMeshEditor.cpp
+++ b/Engine/Source/Editor/DestructibleMeshEditor/Private/DestructibleMeshEditor.cpp
@@ -8,6 +8,8 @@
 #include "Editor/PropertyEditor/Public/PropertyEditorModule.h"
 #include "Editor/PropertyEditor/Public/IDetailsView.h"
 
+#include "Engine/DestructibleMesh.h"
+
 #define LOCTEXT_NAMESPACE "DestructibleMeshEditor"
 
 const FName FDestructibleMeshEditor::ViewportTabId( TEXT( "DestructibleMeshEditor_Viewport" ) );
diff --git a/Engine/Source/Editor/DestructibleMeshEditor/Private/DestructibleMeshEditorModule.cpp b/Engine/Source/Editor/DestructibleMeshEditor/Private/DestructibleMeshEditorModule.cpp
--- a/Engine/Source/Editor/DestructibleMeshEditor/Private/DestructibleMeshEditorModule.cpp
+++ b/Engine/Source/Editor/DestructibleMeshEditor/Private/DestructibleMeshEditorModule.cpp
@@ -8,6 +8,7 @@
 #include "AssetToolsModule.h"
 
 #include "ApexDestructibleAssetImport.h"
+#include "Engine/DestructibleMesh.h"
 
 IMPLEMENT_MODULE( FDestructibleMeshEditorModule, DestructibleMeshEditor );
 
diff --git a/Engine/Source/Editor/DestructibleMeshEditor/Private/SDestructibleMeshEditorViewport.cpp b/Engine/Source/Editor/DestructibleMeshEditor/Private/SDestructibleMeshEditorViewport.cpp
--- a/Engine/Source/Editor/DestructibleMeshEditor/Private/SDestructibleMeshEditorViewport.cpp
+++ b/Engine/Source/Editor/DestructibleMeshEditor/Private/SDestructibleMeshEditorViewport.cpp
@@ -10,6 +10,7 @@
 #include "DesktopPlatformModule.h"
 #include "FbxImporter.h"
 #include "ComponentReregisterContext.h"
+#include "Engine/DestructibleMesh.h"
 
 DEFINE_LOG_CATEGORY_STATIC(LogDestructibleMeshEditor, Log, All);
 
diff --git a/Engine/Source/Editor/UnrealEd/Private/ComponentAssetBroker.cpp b/Engine/Source/Editor/UnrealEd/Private/ComponentAssetBroker.cpp
--- a/Engine/Source/Editor/UnrealEd/Private/ComponentAssetBroker.cpp
+++ b/Engine/Source/Editor/UnrealEd/Private/ComponentAssetBroker.cpp
@@ -8,6 +8,8 @@
 #include "Particles/ParticleSystem.h"
 #include "Particles/ParticleSystemComponent.h"
 
+#include "Engine/DestructibleMesh.h"
+
 //////////////////////////////////////////////////////////////////////////
 // FStaticMeshComponentBroker
 
diff --git a/Engine/Source/Editor/UnrealEd/Private/Factories/ActorFactory.cpp b/Engine/Source/Editor/UnrealEd/Private/Factories/ActorFactory.cpp
--- a/Engine/Source/Editor/UnrealEd/Private/Factories/ActorFactory.cpp
+++ b/Engine/Source/Editor/UnrealEd/Private/Factories/ActorFactory.cpp
@@ -33,6 +33,8 @@ ActorFactory.cpp:
 #include "Engine/TriggerCapsule.h"
 #include "Engine/TextRenderActor.h"
 
+#include "Engine/DestructibleMesh.h"
+
 DEFINE_LOG_CATEGORY_STATIC(LogActorFactory, Log, All);
 
 #define LOCTEXT_NAMESPACE "ActorFactory"
diff --git a/Engine/Source/Editor/UnrealEd/Private/Factories/ApexDestructibleAssetImport.cpp b/Engine/Source/Editor/UnrealEd/Private/Factories/ApexDestructibleAssetImport.cpp
--- a/Engine/Source/Editor/UnrealEd/Private/Factories/ApexDestructibleAssetImport.cpp
+++ b/Engine/Source/Editor/UnrealEd/Private/Factories/ApexDestructibleAssetImport.cpp
@@ -23,6 +23,7 @@ DEFINE_LOG_CATEGORY_STATIC(LogApexDestructibleAssetImport, Log, All);
 #include "ApexDestructibleAssetImport.h"
 #include "Developer/MeshUtilities/Public/MeshUtilities.h"
 #include "ComponentReregisterContext.h"
+#include "Engine/DestructibleMesh.h"
 
 #if WITH_APEX
 
diff --git a/Engine/Source/Editor/UnrealEd/Private/Factories/EditorFactories.cpp b/Engine/Source/Editor/UnrealEd/Private/Factories/EditorFactories.cpp
--- a/Engine/Source/Editor/UnrealEd/Private/Factories/EditorFactories.cpp
+++ b/Engine/Source/Editor/UnrealEd/Private/Factories/EditorFactories.cpp
@@ -46,6 +46,7 @@
 
 #include "Particles/ParticleSystem.h"
 
+#include "Engine/DestructibleMesh.h"
 #if WITH_APEX
 #include "ApexDestructibleAssetImport.h"
 #endif // WITH_APEX
diff --git a/Engine/Source/Editor/UnrealEd/Private/Tests/AutomationEditorCommon.cpp b/Engine/Source/Editor/UnrealEd/Private/Tests/AutomationEditorCommon.cpp
--- a/Engine/Source/Editor/UnrealEd/Private/Tests/AutomationEditorCommon.cpp
+++ b/Engine/Source/Editor/UnrealEd/Private/Tests/AutomationEditorCommon.cpp
@@ -9,6 +9,7 @@
 #include "EngineVersion.h"
 #include "ShaderCompiler.h"
 #include "BehaviorTree/BehaviorTree.h"
+#include "Engine/DestructibleMesh.h"
 
 
 
diff --git a/Engine/Source/Editor/UnrealEd/Private/Tests/EditorOpenAssetAutomationTests.cpp b/Engine/Source/Editor/UnrealEd/Private/Tests/EditorOpenAssetAutomationTests.cpp
--- a/Engine/Source/Editor/UnrealEd/Private/Tests/EditorOpenAssetAutomationTests.cpp
+++ b/Engine/Source/Editor/UnrealEd/Private/Tests/EditorOpenAssetAutomationTests.cpp
@@ -12,6 +12,7 @@
 #include "Materials/MaterialFunction.h"
 #include "Slate/SlateBrushAsset.h"
 #include "BehaviorTree/BehaviorTree.h"
+#include "Engine/DestructibleMesh.h"
 
 
 
diff --git a/Engine/Source/Editor/UnrealEd/Private/UnrealEdSrv.cpp b/Engine/Source/Editor/UnrealEd/Private/UnrealEdSrv.cpp
--- a/Engine/Source/Editor/UnrealEd/Private/UnrealEdSrv.cpp
+++ b/Engine/Source/Editor/UnrealEd/Private/UnrealEdSrv.cpp
@@ -36,6 +36,7 @@
 #include "ComponentReregisterContext.h"
 #include "JsonInternationalizationArchiveSerializer.h"
 #include "JsonInternationalizationManifestSerializer.h"
+#include "Engine/DestructibleMesh.h"
 
 
 DEFINE_LOG_CATEGORY_STATIC(LogUnrealEdSrv, Log, All);
diff --git a/Engine/Source/Runtime/Engine/Private/DestructibleComponent.cpp b/Engine/Source/Runtime/Engine/Private/DestructibleComponent.cpp
--- a/Engine/Source/Runtime/Engine/Private/DestructibleComponent.cpp
+++ b/Engine/Source/Runtime/Engine/Private/DestructibleComponent.cpp
@@ -14,6 +14,7 @@
 #include "ObjectEditorUtils.h"
 #include "Particles/ParticleSystemComponent.h"
 #include "PhysicsEngine/DestructibleActor.h"
+#include "Engine/DestructibleMesh.h"
 
 UDestructibleComponent::UDestructibleComponent(const class FPostConstructInitializeProperties& PCIP)
 	: Super(PCIP)
diff --git a/Engine/Source/Runtime/Engine/Private/EnginePrivate.h b/Engine/Source/Runtime/Engine/Private/EnginePrivate.h
--- a/Engine/Source/Runtime/Engine/Private/EnginePrivate.h
+++ b/Engine/Source/Runtime/Engine/Private/EnginePrivate.h
@@ -232,7 +232,6 @@
 #include "Engine/Selection.h"
 #include "Engine/SimpleConstructionScript.h"
 #include "Engine/SkeletalMesh.h"
-#include "Engine/DestructibleMesh.h"
 #include "Engine/SkeletalMeshSocket.h"
 #include "Engine/StaticMesh.h"
 #include "Engine/StaticMeshSocket.h"
diff --git a/Engine/Source/Runtime/Engine/Public/Engine.h b/Engine/Source/Runtime/Engine/Public/Engine.h
--- a/Engine/Source/Runtime/Engine/Public/Engine.h
+++ b/Engine/Source/Runtime/Engine/Public/Engine.h
@@ -283,7 +283,6 @@
 #include "Engine/Selection.h"
 #include "Engine/SimpleConstructionScript.h"
 #include "Animation/PreviewAssetAttachComponent.h"
-#include "Engine/DestructibleMesh.h"
 #include "Engine/SkeletalMeshSocket.h"
 #include "Animation/Skeleton.h"
 #include "Sound/DialogueSoundWaveProxy.h"
