minor color adjustments to the assets (desaturated, RED could be perceives as error)

diff --git a/Engine/Source/Developer/AssetTools/Private/AssetTypeActions/AssetTypeActions_CameraAnim.h b/Engine/Source/Developer/AssetTools/Private/AssetTypeActions/AssetTypeActions_CameraAnim.h
--- a/Engine/Source/Developer/AssetTools/Private/AssetTypeActions/AssetTypeActions_CameraAnim.h
+++ b/Engine/Source/Developer/AssetTools/Private/AssetTypeActions/AssetTypeActions_CameraAnim.h
@@ -7,7 +7,7 @@ class FAssetTypeActions_CameraAnim : public FAssetTypeActions_Base
 public:
 	// IAssetTypeActions Implementation
 	virtual FText GetName() const override { return NSLOCTEXT("AssetTypeActions", "AssetTypeActions_CameraAnim", "Camera Anim"); }
-	virtual FColor GetTypeColor() const override { return FColor(255,0,0); }
+	virtual FColor GetTypeColor() const override { return FColor(128,0,64); }
 	virtual UClass* GetSupportedClass() const override { return UCameraAnim::StaticClass(); }
 	virtual void OpenAssetEditor( const TArray<UObject*>& InObjects, TSharedPtr<class IToolkitHost> EditWithinLevelEditor = TSharedPtr<IToolkitHost>() ) override;
 	virtual bool CanFilter() override { return true; }
diff --git a/Engine/Source/Developer/AssetTools/Private/AssetTypeActions/AssetTypeActions_Font.h b/Engine/Source/Developer/AssetTools/Private/AssetTypeActions/AssetTypeActions_Font.h
--- a/Engine/Source/Developer/AssetTools/Private/AssetTypeActions/AssetTypeActions_Font.h
+++ b/Engine/Source/Developer/AssetTools/Private/AssetTypeActions/AssetTypeActions_Font.h
@@ -7,7 +7,7 @@ class FAssetTypeActions_Font : public FAssetTypeActions_Base
 public:
 	// IAssetTypeActions Implementation
 	virtual FText GetName() const override { return NSLOCTEXT("AssetTypeActions", "AssetTypeActions_Font", "Font"); }
-	virtual FColor GetTypeColor() const override { return FColor(255,0,0); }
+	virtual FColor GetTypeColor() const override { return FColor(128,128,64); }
 	virtual UClass* GetSupportedClass() const override { return UFont::StaticClass(); }
 	virtual bool HasActions(const TArray<UObject*>& InObjects) const override { return true; }
 	virtual void GetActions(const TArray<UObject*>& InObjects, FMenuBuilder& MenuBuilder) override;
diff --git a/Engine/Source/Developer/AssetTools/Private/AssetTypeActions/AssetTypeActions_Material.h b/Engine/Source/Developer/AssetTools/Private/AssetTypeActions/AssetTypeActions_Material.h
--- a/Engine/Source/Developer/AssetTools/Private/AssetTypeActions/AssetTypeActions_Material.h
+++ b/Engine/Source/Developer/AssetTools/Private/AssetTypeActions/AssetTypeActions_Material.h
@@ -7,7 +7,7 @@ class FAssetTypeActions_Material : public FAssetTypeActions_MaterialInterface
 public:
 	// IAssetTypeActions Implementation
 	virtual FText GetName() const override { return NSLOCTEXT("AssetTypeActions", "AssetTypeActions_Material", "Material"); }
-	virtual FColor GetTypeColor() const override { return FColor(0,255,0); }
+	virtual FColor GetTypeColor() const override { return FColor(64,192,64); }
 	virtual UClass* GetSupportedClass() const override { return UMaterial::StaticClass(); }
 	virtual void OpenAssetEditor( const TArray<UObject*>& InObjects, TSharedPtr<class IToolkitHost> EditWithinLevelEditor = TSharedPtr<IToolkitHost>() ) override;
 	virtual bool CanFilter() override { return true; }
diff --git a/Engine/Source/Developer/AssetTools/Private/AssetTypeActions/AssetTypeActions_MaterialInterface.h b/Engine/Source/Developer/AssetTools/Private/AssetTypeActions/AssetTypeActions_MaterialInterface.h
--- a/Engine/Source/Developer/AssetTools/Private/AssetTypeActions/AssetTypeActions_MaterialInterface.h
+++ b/Engine/Source/Developer/AssetTools/Private/AssetTypeActions/AssetTypeActions_MaterialInterface.h
@@ -7,7 +7,7 @@ class FAssetTypeActions_MaterialInterface : public FAssetTypeActions_Base
 public:
 	// IAssetTypeActions Implementation
 	virtual FText GetName() const override { return NSLOCTEXT("AssetTypeActions", "AssetTypeActions_MaterialInterface", "Material Interface"); }
-	virtual FColor GetTypeColor() const override { return FColor(0,255,0); }
+	virtual FColor GetTypeColor() const override { return FColor(64,192,64); }
 	virtual UClass* GetSupportedClass() const override { return UMaterialInterface::StaticClass(); }
 	virtual bool HasActions ( const TArray<UObject*>& InObjects ) const override { return true; }
 	virtual void GetActions( const TArray<UObject*>& InObjects, FMenuBuilder& MenuBuilder ) override;
diff --git a/Engine/Source/Developer/AssetTools/Private/AssetTypeActions/AssetTypeActions_Texture.h b/Engine/Source/Developer/AssetTools/Private/AssetTypeActions/AssetTypeActions_Texture.h
--- a/Engine/Source/Developer/AssetTools/Private/AssetTypeActions/AssetTypeActions_Texture.h
+++ b/Engine/Source/Developer/AssetTools/Private/AssetTypeActions/AssetTypeActions_Texture.h
@@ -7,7 +7,7 @@ class FAssetTypeActions_Texture : public FAssetTypeActions_Base
 public:
 	// IAssetTypeActions Implementation
 	virtual FText GetName() const override { return NSLOCTEXT("AssetTypeActions", "AssetTypeActions_Texture", "BaseTexture"); }
-	virtual FColor GetTypeColor() const override { return FColor(255,0,0); }
+	virtual FColor GetTypeColor() const override { return FColor(192,64,64); }
 	virtual UClass* GetSupportedClass() const override { return UTexture::StaticClass(); }
 	virtual bool HasActions ( const TArray<UObject*>& InObjects ) const override { return true; }
 	virtual void GetActions( const TArray<UObject*>& InObjects, FMenuBuilder& MenuBuilder ) override;
diff --git a/Engine/Source/Developer/AssetTools/Private/AssetTypeActions/AssetTypeActions_Texture2D.h b/Engine/Source/Developer/AssetTools/Private/AssetTypeActions/AssetTypeActions_Texture2D.h
--- a/Engine/Source/Developer/AssetTools/Private/AssetTypeActions/AssetTypeActions_Texture2D.h
+++ b/Engine/Source/Developer/AssetTools/Private/AssetTypeActions/AssetTypeActions_Texture2D.h
@@ -7,7 +7,7 @@ class FAssetTypeActions_Texture2D : public FAssetTypeActions_Texture
 public:
 	// IAssetTypeActions Implementation
 	virtual FText GetName() const override { return NSLOCTEXT("AssetTypeActions", "AssetTypeActions_Texture2D", "Texture"); }
-	virtual FColor GetTypeColor() const override { return FColor(255,0,0); }
+	virtual FColor GetTypeColor() const override { return FColor(192,64,64); }
 	virtual UClass* GetSupportedClass() const override { return UTexture2D::StaticClass(); }
 	virtual bool HasActions ( const TArray<UObject*>& InObjects ) const override { return true; }
 	virtual void GetActions( const TArray<UObject*>& InObjects, FMenuBuilder& MenuBuilder ) override;
diff --git a/Engine/Source/Developer/AssetTools/Private/AssetTypeActions/AssetTypeActions_TextureCube.h b/Engine/Source/Developer/AssetTools/Private/AssetTypeActions/AssetTypeActions_TextureCube.h
--- a/Engine/Source/Developer/AssetTools/Private/AssetTypeActions/AssetTypeActions_TextureCube.h
+++ b/Engine/Source/Developer/AssetTools/Private/AssetTypeActions/AssetTypeActions_TextureCube.h
@@ -7,7 +7,7 @@ class FAssetTypeActions_TextureCube : public FAssetTypeActions_Texture
 public:
 	// IAssetTypeActions Implementation
 	virtual FText GetName() const override { return NSLOCTEXT("AssetTypeActions", "AssetTypeActions_TextureCube", "Texture Cube"); }
-	virtual FColor GetTypeColor() const override { return FColor(255,0,0); }
+	virtual FColor GetTypeColor() const override { return FColor(192,64,128); }
 	virtual UClass* GetSupportedClass() const override { return UTextureCube::StaticClass(); }
 	virtual bool CanFilter() override { return true; }
 };
\ No newline at end of file
diff --git a/Engine/Source/Developer/AssetTools/Private/AssetTypeActions/AssetTypeActions_TextureLightProfile.h b/Engine/Source/Developer/AssetTools/Private/AssetTypeActions/AssetTypeActions_TextureLightProfile.h
--- a/Engine/Source/Developer/AssetTools/Private/AssetTypeActions/AssetTypeActions_TextureLightProfile.h
+++ b/Engine/Source/Developer/AssetTools/Private/AssetTypeActions/AssetTypeActions_TextureLightProfile.h
@@ -7,7 +7,7 @@ class FAssetTypeActions_TextureLightProfile : public FAssetTypeActions_Texture
 public:
 	// IAssetTypeActions Implementation
 	virtual FText GetName() const override { return NSLOCTEXT("AssetTypeActions", "AssetTypeActions_TextureLightProfile", "Texture Light Profile"); }
-	virtual FColor GetTypeColor() const override { return FColor(255,0,0); }
+	virtual FColor GetTypeColor() const override { return FColor(192,64,192); }
 	virtual UClass* GetSupportedClass() const override { return UTextureLightProfile::StaticClass(); }
 	virtual bool CanFilter() override { return true; }
 };
\ No newline at end of file
diff --git a/Engine/Source/Developer/AssetTools/Private/AssetTypeActions/AssetTypeActions_TextureRenderTarget.h b/Engine/Source/Developer/AssetTools/Private/AssetTypeActions/AssetTypeActions_TextureRenderTarget.h
--- a/Engine/Source/Developer/AssetTools/Private/AssetTypeActions/AssetTypeActions_TextureRenderTarget.h
+++ b/Engine/Source/Developer/AssetTools/Private/AssetTypeActions/AssetTypeActions_TextureRenderTarget.h
@@ -7,7 +7,7 @@ class FAssetTypeActions_TextureRenderTarget : public FAssetTypeActions_Texture
 public:
 	// IAssetTypeActions Implementation
 	virtual FText GetName() const override { return NSLOCTEXT("AssetTypeActions", "AssetTypeActions_TextureRenderTarget", "Texture Render Target"); }
-	virtual FColor GetTypeColor() const override { return FColor(255,0,0); }
+	virtual FColor GetTypeColor() const override { return FColor(128,64,64); }
 	virtual UClass* GetSupportedClass() const override { return UTextureRenderTarget::StaticClass(); }
 	virtual bool HasActions ( const TArray<UObject*>& InObjects ) const override { return true; }
 	virtual void GetActions( const TArray<UObject*>& InObjects, FMenuBuilder& MenuBuilder ) override;
