Fixed crash moving assets which are checked out to another user.  Fixed bug where asset move happens before the user is given the chance to cancel in the source control dialog.
#ttp 348063 - Crash moving files in Content Browser
#branch UE4
#proj Developer.AssetTools, Editor.ContentBrowser, Editor.PackagesDialog, Editor.UnrealEd

diff --git a/Engine/Source/Developer/AssetTools/Private/AssetRenameManager.cpp b/Engine/Source/Developer/AssetTools/Private/AssetRenameManager.cpp
--- a/Engine/Source/Developer/AssetTools/Private/AssetRenameManager.cpp
+++ b/Engine/Source/Developer/AssetTools/Private/AssetRenameManager.cpp
@@ -202,8 +202,8 @@ void FAssetRenameManager::FixReferencesAndRename(TArray<FAssetRenameDataWithRefe
 		TArray<UPackage*> ReferencingPackagesToSave;
 		LoadReferencingPackages(AssetsToRename, ReferencingPackagesToSave);
 
-		// Prompt to check out all referencing packages, leave redirectors for assets referenced by packages that are not checked out and remove those packages from the save list.
-		const bool bUserAcceptedCheckout = CheckOutReferencingPackages(AssetsToRename, ReferencingPackagesToSave);
+		// Prompt to check out source package and all referencing packages, leave redirectors for assets referenced by packages that are not checked out and remove those packages from the save list.
+		const bool bUserAcceptedCheckout = CheckOutPackages(AssetsToRename, ReferencingPackagesToSave);
 
 		if ( bUserAcceptedCheckout )
 		{
@@ -215,11 +215,11 @@ void FAssetRenameManager::FixReferencesAndRename(TArray<FAssetRenameDataWithRefe
 
 			// Save all packages that were referencing any of the assets that were moved without redirectors
 			SaveReferencingPackages(ReferencingPackagesToSave);
-
-			// Finally, report any failures that happened during the rename
-			ReportFailures(AssetsToRename);
 		}
 	}
+
+	// Finally, report any failures that happened during the rename
+	ReportFailures(AssetsToRename);
 }
 
 TArray<TWeakObjectPtr<UObject>> FAssetRenameManager::FindCDOReferencedAssets(const TArray<FAssetRenameDataWithReferencers>& AssetsToRename) const
@@ -362,19 +362,6 @@ void FAssetRenameManager::LoadReferencingPackages(TArray<FAssetRenameDataWithRef
 			{
 				// This asset is not local. It is not safe to rename it without leaving a redirector
 				RenameData.bCreateRedirector = true;
-
-				// If this asset is locked or not current, mark it failed to prevent it from being renamed
-				if ( SourceControlState->IsCheckedOutOther() )
-				{
-					RenameData.bRenameFailed = true;
-					RenameData.FailureReason = LOCTEXT("RenameFailedCheckedOutByOther", "Checked out by another user.");
-				}
-				else if ( !SourceControlState->IsCurrent() )
-				{
-					RenameData.bRenameFailed = true;
-					RenameData.FailureReason = LOCTEXT("RenameFailedNotCurrent", "Out of date.");
-				}
-
 				continue;
 			}
 		}
@@ -421,20 +408,39 @@ void FAssetRenameManager::LoadReferencingPackages(TArray<FAssetRenameDataWithRef
 	GWarn->EndSlowTask();
 }
 
-bool FAssetRenameManager::CheckOutReferencingPackages(TArray<FAssetRenameDataWithReferencers>& AssetsToRename, TArray<UPackage*>& InOutReferencingPackagesToSave) const
+bool FAssetRenameManager::CheckOutPackages(TArray<FAssetRenameDataWithReferencers>& AssetsToRename, TArray<UPackage*>& InOutReferencingPackagesToSave) const
 {
 	bool bUserAcceptedCheckout = true;
 	
-	if ( InOutReferencingPackagesToSave.Num() > 0 )
+	// Build list of packages to check out: the source package and any referencing packages (in the case that we do not create a redirector)
+	TArray<UPackage*> PackagesToCheckOut;
+	PackagesToCheckOut.Reset(AssetsToRename.Num() + InOutReferencingPackagesToSave.Num());
+
+	for (const auto& AssetToRename : AssetsToRename)
+	{
+		if (AssetToRename.Asset.IsValid())
+		{
+			PackagesToCheckOut.Add(AssetToRename.Asset->GetOutermost());
+		}
+	}
+
+	for (UPackage* ReferencingPackage : InOutReferencingPackagesToSave)
+	{
+		PackagesToCheckOut.Add(ReferencingPackage);
+	}
+
+	// Check out the packages
+	if (PackagesToCheckOut.Num() > 0)
 	{
 		if ( ISourceControlModule::Get().IsEnabled() )
 		{
 			TArray<UPackage*> PackagesCheckedOutOrMadeWritable;
 			TArray<UPackage*> PackagesNotNeedingCheckout;
-			bUserAcceptedCheckout = FEditorFileUtils::PromptToCheckoutPackages( false, InOutReferencingPackagesToSave, &PackagesCheckedOutOrMadeWritable, &PackagesNotNeedingCheckout );
+			bUserAcceptedCheckout = FEditorFileUtils::PromptToCheckoutPackages( false, PackagesToCheckOut, &PackagesCheckedOutOrMadeWritable, &PackagesNotNeedingCheckout );
 			if ( bUserAcceptedCheckout )
 			{
-				TArray<UPackage*> PackagesThatCouldNotBeCheckedOut = InOutReferencingPackagesToSave;
+				// Make a list of any packages in the list which weren't checked out for some reason
+				TArray<UPackage*> PackagesThatCouldNotBeCheckedOut = PackagesToCheckOut;
 
 				for ( auto PackageIt = PackagesCheckedOutOrMadeWritable.CreateConstIterator(); PackageIt; ++PackageIt )
 				{
@@ -446,20 +452,10 @@ bool FAssetRenameManager::CheckOutReferencingPackages(TArray<FAssetRenameDataWit
 					PackagesThatCouldNotBeCheckedOut.Remove(*PackageIt);
 				}
 
-				for ( auto PackageIt = PackagesThatCouldNotBeCheckedOut.CreateConstIterator(); PackageIt; ++PackageIt )
+				// If there's anything which couldn't be checked out, abort the operation.
+				if (PackagesThatCouldNotBeCheckedOut.Num() > 0)
 				{
-					const FName NonCheckedOutPackageName = (*PackageIt)->GetFName();
-
-					for ( auto RenameDataIt = AssetsToRename.CreateIterator(); RenameDataIt; ++RenameDataIt )
-					{
-						FAssetRenameDataWithReferencers& RenameData = *RenameDataIt;
-						if ( RenameData.ReferencingPackageNames.Contains(NonCheckedOutPackageName) )
-						{
-							RenameData.bCreateRedirector = true;
-						}
-					}
-
-					InOutReferencingPackagesToSave.Remove(*PackageIt);
+					bUserAcceptedCheckout = false;
 				}
 			}
 		}
@@ -655,11 +651,18 @@ void FAssetRenameManager::PerformAssetRename(TArray<FAssetRenameDataWithReferenc
 	{
 		const bool bCheckDirty = false;
 		const bool bPromptToSave = false;
-		FEditorFileUtils::PromptForCheckoutAndSave(PackagesToSave, bCheckDirty, bPromptToSave);
+		const bool bAlreadyCheckedOut = true;
+		FEditorFileUtils::PromptForCheckoutAndSave(PackagesToSave, bCheckDirty, bPromptToSave, nullptr, bAlreadyCheckedOut);
 
 		ISourceControlModule::Get().QueueStatusUpdate(PackagesToSave);
 	}
 
+	// Now branch the files in source control if possible
+	for (const auto& AssetToRename : AssetsToRename)
+	{
+		SourceControlHelpers::BranchPackage(AssetToRename.Asset->GetOutermost(), FindPackage(nullptr, *AssetToRename.OriginalAssetPath));
+	}
+
 	// Clean up all packages that were left empty
 	if ( PotentialPackagesToDelete.Num() > 0 )
 	{
diff --git a/Engine/Source/Developer/AssetTools/Private/AssetRenameManager.h b/Engine/Source/Developer/AssetTools/Private/AssetRenameManager.h
--- a/Engine/Source/Developer/AssetTools/Private/AssetRenameManager.h
+++ b/Engine/Source/Developer/AssetTools/Private/AssetRenameManager.h
@@ -48,11 +48,11 @@ private:
 	void LoadReferencingPackages(TArray<FAssetRenameDataWithReferencers>& AssetsToRename, TArray<UPackage*>& OutReferencingPackagesToSave) const;
 
 	/** 
-	  * Prompts to check out referencing packages and marks assets whose referencing packages were not checked out to leave a redirector.
+	  * Prompts to check out the source package and all referencing packages and marks assets whose referencing packages were not checked out to leave a redirector.
 	  * Trims PackagesToSave when necessary.
 	  * Returns true if the user opted to continue the operation or no dialog was required.
 	  */
-	bool CheckOutReferencingPackages(TArray<FAssetRenameDataWithReferencers>& AssetsToRename, TArray<UPackage*>& InOutReferencingPackagesToSave) const;
+	bool CheckOutPackages(TArray<FAssetRenameDataWithReferencers>& AssetsToRename, TArray<UPackage*>& InOutReferencingPackagesToSave) const;
 
 	/** Finds any read only packages and removes them from the save list. Assets referenced by these packages will leave redirectors. */ 
 	void DetectReadOnlyPackages(TArray<FAssetRenameDataWithReferencers>& AssetsToRename, TArray<UPackage*>& InOutReferencingPackagesToSave) const;
diff --git a/Engine/Source/Editor/ContentBrowser/Private/ContentBrowserUtils.cpp b/Engine/Source/Editor/ContentBrowser/Private/ContentBrowserUtils.cpp
--- a/Engine/Source/Editor/ContentBrowser/Private/ContentBrowserUtils.cpp
+++ b/Engine/Source/Editor/ContentBrowser/Private/ContentBrowserUtils.cpp
@@ -503,17 +503,6 @@ void ContentBrowserUtils::MoveAssets(const TArray<UObject*>& Assets, const FStri
 	{
 		AssetToolsModule.Get().RenameAssets(AssetsAndNames);
 	}
-
-	// Now branch the files in source control if possible
-	for(const auto& AssetAndName : AssetsAndNames)
-	{
-		check(AssetAndName.Asset.Get());
-		UPackage* DestPackage = FindPackage(nullptr, *(AssetAndName.PackagePath / AssetAndName.NewName));
-		UPackage* SourcePackage = FindPackage(nullptr, *AssetAndName.OriginalAssetPath);
-		check(DestPackage);
-		check(SourcePackage);
-		SourceControlHelpers::BranchPackage(DestPackage, SourcePackage);
-	}
 }
 
 int32 ContentBrowserUtils::DeleteAssets(const TArray<UObject*>& AssetsToDelete)
diff --git a/Engine/Source/Editor/PackagesDialog/Private/SPackagesDialog.cpp b/Engine/Source/Editor/PackagesDialog/Private/SPackagesDialog.cpp
--- a/Engine/Source/Editor/PackagesDialog/Private/SPackagesDialog.cpp
+++ b/Engine/Source/Editor/PackagesDialog/Private/SPackagesDialog.cpp
@@ -555,7 +555,10 @@ void SPackagesDialog::PopulateIgnoreForSaveItems( const TSet<FString>& InIgnoreP
 
 		const ESlateCheckBoxState::Type CheckedStatus = (InIgnorePackages.Find(ItemName) != NULL) ? ESlateCheckBoxState::Unchecked : ESlateCheckBoxState::Checked;
 
-		(*ItItem)->SetState( CheckedStatus );
+		if (!(*ItItem)->IsDisabled())
+		{
+			(*ItItem)->SetState(CheckedStatus);
+		}
 	}
 }
 
diff --git a/Engine/Source/Editor/UnrealEd/Private/FileHelpers.cpp b/Engine/Source/Editor/UnrealEd/Private/FileHelpers.cpp
--- a/Engine/Source/Editor/UnrealEd/Private/FileHelpers.cpp
+++ b/Engine/Source/Editor/UnrealEd/Private/FileHelpers.cpp
@@ -2626,6 +2626,7 @@ bool FEditorFileUtils::SaveCurrentLevel()
  * @param		bCheckDirty					If true, only packages that are dirty in PackagesToSave will be saved	
  * @param		bPromptToSave				If true the user will be prompted with a list of packages to save, otherwise all passed in packages are saved
  * @param		OutFailedPackages			[out] If specified, will be filled in with all of the packages that failed to save successfully
+ * @param		bAlreadyCheckedOut			If true, the user will not be prompted with the source control dialog
  *
  * @return		An enum value signifying success, failure, user declined, or cancellation. If any packages at all failed to save during execution, the return code will be 
  *				failure, even if other packages successfully saved. If the user cancels at any point during any prompt, the return code will be cancellation, even though it
@@ -2633,7 +2634,7 @@ bool FEditorFileUtils::SaveCurrentLevel()
  *				Save" option on the dialog, the return code will indicate the user has declined out of the prompt. This way calling code can distinguish between a decline and a cancel
  *				and then proceed as planned, or abort its operation accordingly.
  */
-FEditorFileUtils::EPromptReturnCode FEditorFileUtils::PromptForCheckoutAndSave( const TArray<UPackage*>& InPackages, bool bCheckDirty, bool bPromptToSave, TArray<UPackage*>* OutFailedPackages )
+FEditorFileUtils::EPromptReturnCode FEditorFileUtils::PromptForCheckoutAndSave( const TArray<UPackage*>& InPackages, bool bCheckDirty, bool bPromptToSave, TArray<UPackage*>* OutFailedPackages, bool bAlreadyCheckedOut )
 {
 	// Check for re-entrance into this function
 	if ( bIsPromptingForCheckoutAndSave )
@@ -2786,7 +2787,7 @@ FEditorFileUtils::EPromptReturnCode FEditorFileUtils::PromptForCheckoutAndSave(
 		TArray<UPackage*> PackagesNotNeedingCheckout;
 
 		// Prompt to check-out any packages under source control
-		const bool UserResponse = FEditorFileUtils::PromptToCheckoutPackages( false, PackagesToSave, &PackagesCheckedOutOrMadeWritable, &PackagesNotNeedingCheckout );
+		const bool UserResponse = bAlreadyCheckedOut || FEditorFileUtils::PromptToCheckoutPackages( false, PackagesToSave, &PackagesCheckedOutOrMadeWritable, &PackagesNotNeedingCheckout );
 
 		if( UserResponse || PackagesNotNeedingCheckout.Num() > 0 )
 		{
@@ -2795,7 +2796,14 @@ FEditorFileUtils::EPromptReturnCode FEditorFileUtils::PromptForCheckoutAndSave(
 
 			if ( UserResponse )
 			{
-				FinalSaveList.Append(PackagesCheckedOutOrMadeWritable);
+				if (bAlreadyCheckedOut)
+				{
+					FinalSaveList.Append(PackagesToSave);
+				}
+				else
+				{
+					FinalSaveList.Append(PackagesCheckedOutOrMadeWritable);
+				}
 			}
 
 			const FScopedBusyCursor BusyCursor;
diff --git a/Engine/Source/Editor/UnrealEd/Public/FileHelpers.h b/Engine/Source/Editor/UnrealEd/Public/FileHelpers.h
--- a/Engine/Source/Editor/UnrealEd/Public/FileHelpers.h
+++ b/Engine/Source/Editor/UnrealEd/Public/FileHelpers.h
@@ -205,6 +205,7 @@ public:
 	 * @param		bCheckDirty					If true, only packages that are dirty in PackagesToSave will be saved	
 	 * @param		bPromptToSave				If true the user will be prompted with a list of packages to save, otherwise all passed in packages are saved
 	 * @param		OutFailedPackages			[out] If specified, will be filled in with all of the packages that failed to save successfully
+	 * @param		bAlreadyCheckedOut			If true, the user will not be prompted with the source control dialog
 	 *
 	 * @return		An enum value signifying success, failure, user declined, or cancellation. If any packages at all failed to save during execution, the return code will be 
 	 *				failure, even if other packages successfully saved. If the user cancels at any point during any prompt, the return code will be cancellation, even though it
@@ -212,7 +213,7 @@ public:
 	 *				Save" option on the dialog, the return code will indicate the user has declined out of the prompt. This way calling code can distinguish between a decline and a cancel
 	 *				and then proceed as planned, or abort its operation accordingly.
 	 */
-	UNREALED_API static EPromptReturnCode PromptForCheckoutAndSave( const TArray<UPackage*>& PackagesToSave, bool bCheckDirty, bool bPromptToSave, TArray<UPackage*>* OutFailedPackages = NULL );
+	UNREALED_API static EPromptReturnCode PromptForCheckoutAndSave( const TArray<UPackage*>& PackagesToSave, bool bCheckDirty, bool bPromptToSave, TArray<UPackage*>* OutFailedPackages = NULL, bool bAlreadyCheckedOut = false );
 
 	////////////////////////////////////////////////////////////////////////////
 	// Import/Export
