Fixed retargeting crash when preview mesh is missing.

https://jira.ol.epicgames.net/browse/UE-4484?filter=12278
https://jira.ol.epicgames.net/browse/UE-4638?filter=12278

diff --git a/Engine/Source/Developer/AssetTools/Private/AssetTypeActions/AssetTypeActions_AnimBlueprint.cpp b/Engine/Source/Developer/AssetTools/Private/AssetTypeActions/AssetTypeActions_AnimBlueprint.cpp
--- a/Engine/Source/Developer/AssetTools/Private/AssetTypeActions/AssetTypeActions_AnimBlueprint.cpp
+++ b/Engine/Source/Developer/AssetTools/Private/AssetTypeActions/AssetTypeActions_AnimBlueprint.cpp
@@ -149,7 +149,24 @@ void FAssetTypeActions_AnimBlueprint::ExecuteFindSkeleton(TArray<TWeakObjectPtr<
 
 void FAssetTypeActions_AnimBlueprint::RetargetAnimationHandler(USkeleton* OldSkeleton, USkeleton* NewSkeleton, bool bRemapReferencedAssets, bool bConvertSpaces, bool bDuplicateAssets, TArray<TWeakObjectPtr<UObject>> AnimBlueprints)
 {
-	EditorAnimUtils::RetargetAnimations(OldSkeleton, NewSkeleton, AnimBlueprints, bRemapReferencedAssets, bDuplicateAssets, bConvertSpaces);
+	if((!OldSkeleton || OldSkeleton->GetPreviewMesh(true)) && (!NewSkeleton || NewSkeleton->GetPreviewMesh(true)))
+	{
+		EditorAnimUtils::RetargetAnimations(OldSkeleton, NewSkeleton, AnimBlueprints, bRemapReferencedAssets, bDuplicateAssets, bConvertSpaces);
+	}
+	else
+	{
+		FFormatNamedArguments Args;
+		Args.Add(TEXT("OldSkeletonName"), FText::FromString(GetNameSafe(OldSkeleton)));
+		Args.Add(TEXT("NewSkeletonName"), FText::FromString(GetNameSafe(NewSkeleton)));
+		FNotificationInfo Info(FText::Format(LOCTEXT("Retarget Failed", "Old Skeleton {OldSkeletonName} and New Skeleton {NewSkeletonName} need to have Preview Mesh set up to convert animation"), Args));
+		Info.ExpireDuration = 5.0f;
+		Info.bUseLargeFont = false;
+		TSharedPtr<SNotificationItem> Notification = FSlateNotificationManager::Get().AddNotification(Info);
+		if(Notification.IsValid())
+		{
+			Notification->SetCompletionState(SNotificationItem::CS_Fail);
+		}
+	}
 }
 
 void FAssetTypeActions_AnimBlueprint::RetargetAssets(TArray<UObject*> InAnimBlueprints, bool bDuplicateAssets)
diff --git a/Engine/Source/Developer/AssetTools/Private/AssetTypeActions/AssetTypeActions_AnimationAsset.cpp b/Engine/Source/Developer/AssetTools/Private/AssetTypeActions/AssetTypeActions_AnimationAsset.cpp
--- a/Engine/Source/Developer/AssetTools/Private/AssetTypeActions/AssetTypeActions_AnimationAsset.cpp
+++ b/Engine/Source/Developer/AssetTools/Private/AssetTypeActions/AssetTypeActions_AnimationAsset.cpp
@@ -155,7 +155,24 @@ void FAssetTypeActions_AnimationAsset::ExecuteFindSkeleton(TArray<TWeakObjectPtr
 
 void FAssetTypeActions_AnimationAsset::RetargetAnimationHandler(USkeleton* OldSkeleton, USkeleton* NewSkeleton, bool bRemapReferencedAssets, bool bConvertSpaces, bool bDuplicateAssets,  TArray<TWeakObjectPtr<UObject>> InAnimAssets)
 {
-	EditorAnimUtils::RetargetAnimations(OldSkeleton, NewSkeleton, InAnimAssets, bRemapReferencedAssets, bDuplicateAssets, bConvertSpaces);
+	if((!OldSkeleton || OldSkeleton->GetPreviewMesh(false)) && (!NewSkeleton || NewSkeleton->GetPreviewMesh(false)))
+	{
+		EditorAnimUtils::RetargetAnimations(OldSkeleton, NewSkeleton, InAnimAssets, bRemapReferencedAssets, bDuplicateAssets, bConvertSpaces);
+	}
+	else
+	{
+		FFormatNamedArguments Args;
+		Args.Add(TEXT("OldSkeletonName"), FText::FromString(GetNameSafe(OldSkeleton)));
+		Args.Add(TEXT("NewSkeletonName"), FText::FromString(GetNameSafe(NewSkeleton)));
+		FNotificationInfo Info(FText::Format(LOCTEXT("Retarget Failed", "Old Skeleton {OldSkeletonName} and New Skeleton {NewSkeletonName} need to have Preview Mesh set up to convert animation"), Args));
+		Info.ExpireDuration = 5.0f;
+		Info.bUseLargeFont = false;
+		TSharedPtr<SNotificationItem> Notification = FSlateNotificationManager::Get().AddNotification(Info);
+		if(Notification.IsValid())
+		{
+			Notification->SetCompletionState(SNotificationItem::CS_Fail);
+		}
+	}
 }
 
 void FAssetTypeActions_AnimationAsset::RetargetAssets(TArray<UObject*> InAnimAssets, bool bDuplicateAssets)
diff --git a/Engine/Source/Developer/AssetTools/Private/AssetTypeActions/AssetTypeActions_Skeleton.cpp b/Engine/Source/Developer/AssetTools/Private/AssetTypeActions/AssetTypeActions_Skeleton.cpp
--- a/Engine/Source/Developer/AssetTools/Private/AssetTypeActions/AssetTypeActions_Skeleton.cpp
+++ b/Engine/Source/Developer/AssetTools/Private/AssetTypeActions/AssetTypeActions_Skeleton.cpp
@@ -332,7 +332,7 @@ FCreateRigDlg::EResult FCreateRigDlg::ShowModal()
 
 	// Make a list of all skeleton bone list
 	const FReferenceSkeleton& RefSkeleton = Skeleton->GetReferenceSkeleton();
-	const TArray<FBoneNode> & BoneTree = Skeleton->GetBoneTree();
+	const TArray<FBoneNode>& BoneTree = Skeleton->GetBoneTree();
 	for ( int32 BoneTreeId=0; BoneTreeId<RefSkeleton.GetNum(); ++BoneTreeId )
 	{
 		const FName& BoneName = RefSkeleton.GetBoneName(BoneTreeId);
@@ -519,6 +519,23 @@ void FAssetTypeActions_Skeleton::CreateRig(const TWeakObjectPtr<USkeleton> Skele
 
 void FAssetTypeActions_Skeleton::RetargetAnimationHandler(USkeleton* OldSkeleton, USkeleton* NewSkeleton, bool bRemapReferencedAssets, bool bConvertSpaces)
 {
+	if((OldSkeleton && OldSkeleton->GetPreviewMesh(true) == NULL) || (NewSkeleton && NewSkeleton->GetPreviewMesh(true)==NULL))
+	{
+		FFormatNamedArguments Args;
+		Args.Add(TEXT("OldSkeletonName"), FText::FromString(GetNameSafe(OldSkeleton)));
+		Args.Add(TEXT("NewSkeletonName"), FText::FromString(GetNameSafe(NewSkeleton)));
+		FNotificationInfo Info(FText::Format(LOCTEXT("Retarget Failed", "Old Skeleton {OldSkeletonName} and New Skeleton {NewSkeletonName} need to have Preview Mesh set up to convert animation"), Args));
+		Info.ExpireDuration = 5.0f;
+		Info.bUseLargeFont = false;
+		TSharedPtr<SNotificationItem> Notification = FSlateNotificationManager::Get().AddNotification(Info);
+		if(Notification.IsValid())
+		{
+			Notification->SetCompletionState(SNotificationItem::CS_Fail);
+		}
+
+		return;
+	}
+
 	// find all assets who references old skeleton
 	TArray<FName> Packages;
 
diff --git a/Engine/Source/Editor/UnrealEd/Private/SSkeletonWidget.cpp b/Engine/Source/Editor/UnrealEd/Private/SSkeletonWidget.cpp
--- a/Engine/Source/Editor/UnrealEd/Private/SSkeletonWidget.cpp
+++ b/Engine/Source/Editor/UnrealEd/Private/SSkeletonWidget.cpp
@@ -686,7 +686,25 @@ FReply SAnimationRemapSkeleton::OnApply()
 {
 	if (OnRetargetAnimationDelegate.IsBound())
 	{
-		OnRetargetAnimationDelegate.Execute(OldSkeleton, NewSkeleton, bRemapReferencedAssets, bConvertSpaces);
+		// verify if old skeleton and new skeleton has preview mesh, otherwise, it won't work. 
+		if ( (!OldSkeleton || OldSkeleton->GetPreviewMesh(true)) && (!NewSkeleton || NewSkeleton->GetPreviewMesh(true)))
+		{
+			OnRetargetAnimationDelegate.Execute(OldSkeleton, NewSkeleton, bRemapReferencedAssets, bConvertSpaces);
+		}
+		else
+		{
+			FFormatNamedArguments Args;
+			Args.Add(TEXT("OldSkeletonName"), FText::FromString(GetNameSafe(OldSkeleton)));
+			Args.Add(TEXT("NewSkeletonName"), FText::FromString(GetNameSafe(NewSkeleton)));
+			FNotificationInfo Info(FText::Format(LOCTEXT("Retarget Failed", "Old Skeleton {OldSkeletonName} and New Skeleton {NewSkeletonName} need to have Preview Mesh set up to convert animation"), Args));
+			Info.ExpireDuration = 5.0f;
+			Info.bUseLargeFont = false;
+			TSharedPtr<SNotificationItem> Notification = FSlateNotificationManager::Get().AddNotification(Info);
+			if(Notification.IsValid())
+			{
+				Notification->SetCompletionState(SNotificationItem::CS_Fail);
+			}
+		}
 	}
 
 	CloseWindow();
