Creating an asset with the same name as one marked for delete no longer prompts for a check in, but just reverts the delete, checks out the package, and creates the new package in place of the old one.

diff --git a/Engine/Source/Developer/AssetTools/Private/AssetTools.cpp b/Engine/Source/Developer/AssetTools/Private/AssetTools.cpp
--- a/Engine/Source/Developer/AssetTools/Private/AssetTools.cpp
+++ b/Engine/Source/Developer/AssetTools/Private/AssetTools.cpp
@@ -1135,31 +1135,18 @@ bool FAssetTools::CheckForDeletedPackage(const UPackage* Package) const
 			FSourceControlStatePtr SourceControlState = SourceControlProvider.GetState(Package, EStateCacheUsage::ForceUpdate);
 			if ( SourceControlState.IsValid() && SourceControlState->IsDeleted() )
 			{
-				// Creating an asset in a package that is marked for delete.
-				// Prompt the user to check in the package.
-				bool bWantCheckin =
-				EAppReturnType::Yes == FMessageDialog::Open(
-				EAppMsgType::YesNo,
-				FText::Format(
-					LOCTEXT("CheckInDeletedPackage", "'{0}' is marked for delete in source control. You must check in the delete request before creating an asset in its place. Do you want to check in the delete?"),
-					FText::FromString( FPackageName::GetLongPackageAssetName(Package->GetName()) )
-					) 
-				);
-
-				if( bWantCheckin )
+				// Creating an asset in a package that is marked for delete - revert the delete and check out the package
+				if (!SourceControlProvider.Execute(ISourceControlOperation::Create<FRevert>(), Package))
 				{
-					TSharedRef<FCheckIn, ESPMode::ThreadSafe> CheckInOperation = ISourceControlOperation::Create<FCheckIn>();
-					CheckInOperation->SetDescription( LOCTEXT("AutomaticDeleteDesc", "Automatic Delete") );
-					if ( !SourceControlProvider.Execute(CheckInOperation, Package) )
-					{
-						// Failed to check in file
-						FMessageDialog::Open(EAppMsgType::Ok, LOCTEXT("DeletedFileCheckinFailed", "Failed to check in deleted package."));
-						return false;
-					}
+					// Failed to revert file which was marked for delete
+					FMessageDialog::Open(EAppMsgType::Ok, LOCTEXT("RevertDeletedFileFailed", "Failed to revert package which was marked for delete."));
+					return false;
 				}
-				else
+
+				if (!SourceControlProvider.Execute(ISourceControlOperation::Create<FCheckOut>(), Package))
 				{
-					// User chose not to check in the package that is marked for delete
+					// Failed to check out file
+					FMessageDialog::Open(EAppMsgType::Ok, LOCTEXT("CheckOutFileFailed", "Failed to check out package"));
 					return false;
 				}
 			}
