Fixes for non-unity

diff --git a/Engine/Source/Developer/AssetTools/Private/AssetTypeActions/AssetTypeActions_AnimBlueprint.cpp b/Engine/Source/Developer/AssetTools/Private/AssetTypeActions/AssetTypeActions_AnimBlueprint.cpp
--- a/Engine/Source/Developer/AssetTools/Private/AssetTypeActions/AssetTypeActions_AnimBlueprint.cpp
+++ b/Engine/Source/Developer/AssetTools/Private/AssetTypeActions/AssetTypeActions_AnimBlueprint.cpp
@@ -3,6 +3,8 @@
 #include "AssetToolsPrivatePCH.h"
 #include "PersonaModule.h"
 #include "EditorAnimUtils.h"
+#include "NotificationManager.h"
+#include "SNotificationList.h"
 #include "SSkeletonWidget.h"
 
 #define LOCTEXT_NAMESPACE "AssetTypeActions"
diff --git a/Engine/Source/Developer/AssetTools/Private/AssetTypeActions/AssetTypeActions_AnimationAsset.cpp b/Engine/Source/Developer/AssetTools/Private/AssetTypeActions/AssetTypeActions_AnimationAsset.cpp
--- a/Engine/Source/Developer/AssetTools/Private/AssetTypeActions/AssetTypeActions_AnimationAsset.cpp
+++ b/Engine/Source/Developer/AssetTools/Private/AssetTypeActions/AssetTypeActions_AnimationAsset.cpp
@@ -4,6 +4,8 @@
 #include "Toolkits/IToolkitHost.h"
 #include "PersonaModule.h"
 #include "EditorAnimUtils.h"
+#include "NotificationManager.h"
+#include "SNotificationList.h"
 #include "SSkeletonWidget.h"
 
 #define LOCTEXT_NAMESPACE "AssetTypeActions"
diff --git a/Engine/Source/Developer/AssetTools/Private/AssetTypeActions/AssetTypeActions_Skeleton.cpp b/Engine/Source/Developer/AssetTools/Private/AssetTypeActions/AssetTypeActions_Skeleton.cpp
--- a/Engine/Source/Developer/AssetTools/Private/AssetTypeActions/AssetTypeActions_Skeleton.cpp
+++ b/Engine/Source/Developer/AssetTools/Private/AssetTypeActions/AssetTypeActions_Skeleton.cpp
@@ -8,6 +8,8 @@
 #include "Editor/Persona/Public/PersonaModule.h"
 #include "AssetToolsModule.h"
 #include "ContentBrowserModule.h"
+#include "NotificationManager.h"
+#include "SNotificationList.h"
 #include "SSkeletonWidget.h"
 #include "AnimationEditorUtils.h"
 
diff --git a/Engine/Source/Editor/UMGEditor/Private/Customizations/WidgetNavigationCustomization.h b/Engine/Source/Editor/UMGEditor/Private/Customizations/WidgetNavigationCustomization.h
--- a/Engine/Source/Editor/UMGEditor/Private/Customizations/WidgetNavigationCustomization.h
+++ b/Engine/Source/Editor/UMGEditor/Private/Customizations/WidgetNavigationCustomization.h
@@ -2,6 +2,8 @@
 
 #pragma once
 
+#include "WidgetBlueprintEditor.h"
+
 class FWidgetNavigationCustomization : public IPropertyTypeCustomization
 {
 public:
