Fixed crash with setting non skeleton asset to new skeleton

diff --git a/Engine/Source/Developer/AssetTools/Private/AssetTypeActions/AssetTypeActions_AnimationAsset.cpp b/Engine/Source/Developer/AssetTools/Private/AssetTypeActions/AssetTypeActions_AnimationAsset.cpp
--- a/Engine/Source/Developer/AssetTools/Private/AssetTypeActions/AssetTypeActions_AnimationAsset.cpp
+++ b/Engine/Source/Developer/AssetTools/Private/AssetTypeActions/AssetTypeActions_AnimationAsset.cpp
@@ -53,7 +53,7 @@ void FAssetTypeActions_AnimationAsset::FillRetargetMenu( FMenuBuilder& MenuBuild
 			LOCTEXT("AnimAsset_RetargetSkeletonInPlaceTooltip", "Retargets the selected Anim Assets to a new skeleton (and optionally all referenced animations too)"),
 			FSlateIcon(),
 			FUIAction(
-			FExecuteAction::CreateSP( this, &FAssetTypeActions_AnimationAsset::RetargetAssets, InObjects, false ), // false = do not duplicate assets first
+			FExecuteAction::CreateSP( this, &FAssetTypeActions_AnimationAsset::RetargetAssets, InObjects, false, false, TSharedPtr<IToolkitHost>() ), // false = do not duplicate assets first
 			FCanExecuteAction()
 			)
 			);
@@ -64,7 +64,7 @@ void FAssetTypeActions_AnimationAsset::FillRetargetMenu( FMenuBuilder& MenuBuild
 		LOCTEXT("AnimAsset_DuplicateAndRetargetSkeletonTooltip", "Duplicates and then retargets the selected Anim Assets to a new skeleton (and optionally all referenced animations too)"),
 		FSlateIcon(),
 		FUIAction(
-		FExecuteAction::CreateSP( this, &FAssetTypeActions_AnimationAsset::RetargetAssets, InObjects, true ), // true = duplicate assets and retarget them
+		FExecuteAction::CreateSP( this, &FAssetTypeActions_AnimationAsset::RetargetAssets, InObjects, true, false, TSharedPtr<IToolkitHost>() ), // true = duplicate assets and retarget them
 		FCanExecuteAction()
 		)
 		);
@@ -105,12 +105,10 @@ void FAssetTypeActions_AnimationAsset::OpenAssetEditor( const TArray<UObject*>&
 					bool bDuplicateAssets = false;
 					TArray<UObject*> AnimAssets;
 					AnimAssets.Add(AnimAsset);
-					RetargetAssets(AnimAssets, bDuplicateAssets);
+					RetargetAssets(AnimAssets, bDuplicateAssets, true, EditWithinLevelEditor);
 				}
 			}
-
-			AnimSkeleton = AnimAsset->GetSkeleton();
-			if (AnimSkeleton)
+			else if (AnimSkeleton)
 			{
 				const bool bBringToFrontIfOpen = false;
 				if (IAssetEditorInstance* EditorInstance = FAssetEditorManager::Get().FindEditorForAsset(AnimSkeleton, bBringToFrontIfOpen))
@@ -125,10 +123,6 @@ void FAssetTypeActions_AnimationAsset::OpenAssetEditor( const TArray<UObject*>&
 					PersonaModule.CreatePersona(Mode, EditWithinLevelEditor, AnimSkeleton, NULL, AnimAsset, NULL);
 				}
 			}
-			else
-			{
-				FMessageDialog::Open( EAppMsgType::Ok, LOCTEXT("FailedToLoadSkeletonlessAnimAsset", "The Anim Asset could not be loaded because it's skeleton is missing."));
-			}
 		}
 	}
 }
@@ -177,7 +171,41 @@ void FAssetTypeActions_AnimationAsset::RetargetAnimationHandler(USkeleton* OldSk
 	}
 }
 
-void FAssetTypeActions_AnimationAsset::RetargetAssets(TArray<UObject*> InAnimAssets, bool bDuplicateAssets)
+void FAssetTypeActions_AnimationAsset::RetargetNonSkeletonAnimationHandler(USkeleton* OldSkeleton, USkeleton* NewSkeleton, bool bRemapReferencedAssets, bool bConvertSpaces, bool bDuplicateAssets, TArray<TWeakObjectPtr<UObject>> InAnimAssets, TWeakPtr<IToolkitHost> EditWithinLevelEditor)
+{
+	RetargetAnimationHandler(OldSkeleton, NewSkeleton, bRemapReferencedAssets, bConvertSpaces, bDuplicateAssets, InAnimAssets);
+
+	if(NewSkeleton)
+	{
+		for(auto Asset : InAnimAssets)
+		{
+			if (Asset.IsValid())
+			{
+				const bool bBringToFrontIfOpen = false;
+				if(IAssetEditorInstance* EditorInstance = FAssetEditorManager::Get().FindEditorForAsset(NewSkeleton, bBringToFrontIfOpen))
+				{
+					// The skeleton is already open in an editor.
+					// Tell persona that an animation asset was requested
+					EditorInstance->FocusWindow(Asset.Get());
+				}
+				else
+				{
+					EToolkitMode::Type Mode = EditWithinLevelEditor.IsValid() ? EToolkitMode::WorldCentric : EToolkitMode::Standalone;
+					TSharedPtr<IToolkitHost> EditWithInEditor = EditWithinLevelEditor.IsValid()? EditWithinLevelEditor.Pin() : NULL;
+					FPersonaModule& PersonaModule = FModuleManager::LoadModuleChecked<FPersonaModule>("Persona");
+					PersonaModule.CreatePersona(Mode, EditWithInEditor, NewSkeleton, NULL, Cast<UAnimationAsset>(Asset.Get()), NULL);
+				}
+			}
+		}
+	}
+	else
+	{
+		FMessageDialog::Open(EAppMsgType::Ok, LOCTEXT("FailedToLoadSkeletonlessAnimAsset", "The Anim Asset could not be loaded because it's skeleton is missing."));
+	}
+}
+
+
+void FAssetTypeActions_AnimationAsset::RetargetAssets(TArray<UObject*> InAnimAssets, bool bDuplicateAssets, bool bOpenEditor, TSharedPtr<class IToolkitHost> EditWithinLevelEditor)
 {
 	bool bRemapReferencedAssets = false;
 	USkeleton* NewSkeleton = NULL;
@@ -191,7 +219,15 @@ void FAssetTypeActions_AnimationAsset::RetargetAssets(TArray<UObject*> InAnimAss
 	const FText Message = LOCTEXT("SelectSkeletonToRemap", "Select the skeleton to remap this asset to.");
 
 	auto AnimAssets = GetTypedWeakObjectPtrs<UObject>(InAnimAssets);
-	SAnimationRemapSkeleton::ShowWindow(OldSkeleton, Message, FOnRetargetAnimation::CreateSP(this, &FAssetTypeActions_AnimationAsset::RetargetAnimationHandler, bDuplicateAssets, AnimAssets) );
+
+	if (bOpenEditor)
+	{
+		SAnimationRemapSkeleton::ShowWindow(OldSkeleton, Message, FOnRetargetAnimation::CreateSP(this, &FAssetTypeActions_AnimationAsset::RetargetNonSkeletonAnimationHandler, bDuplicateAssets, AnimAssets, TWeakPtr<class IToolkitHost> (EditWithinLevelEditor)) );
+	}
+	else
+	{
+		SAnimationRemapSkeleton::ShowWindow(OldSkeleton, Message, FOnRetargetAnimation::CreateSP(this, &FAssetTypeActions_AnimationAsset::RetargetAnimationHandler, bDuplicateAssets, AnimAssets) );
+	}
 }
 
 
diff --git a/Engine/Source/Developer/AssetTools/Private/AssetTypeActions/AssetTypeActions_AnimationAsset.h b/Engine/Source/Developer/AssetTools/Private/AssetTypeActions/AssetTypeActions_AnimationAsset.h
--- a/Engine/Source/Developer/AssetTools/Private/AssetTypeActions/AssetTypeActions_AnimationAsset.h
+++ b/Engine/Source/Developer/AssetTools/Private/AssetTypeActions/AssetTypeActions_AnimationAsset.h
@@ -24,8 +24,9 @@ private:
 	void ExecuteFindSkeleton(TArray<TWeakObjectPtr<UAnimationAsset>> Objects);
 
 	/** Context menu item handler for changing the supplied assets skeletons */ 
-	void RetargetAssets(TArray<UObject*> InAnimBlueprints, bool bDuplicateAssets);
+	void RetargetAssets(TArray<UObject*> InAnimBlueprints, bool bDuplicateAssets, bool bOpenEditor, TSharedPtr<class IToolkitHost> EditWithinLevelEditor);
 
 	/** Handler for retargeting */
 	void RetargetAnimationHandler(USkeleton* OldSkeleton, USkeleton* NewSkeleton, bool bRemapReferencedAssets, bool bConvertSpaces, bool bDuplicateAssets, TArray<TWeakObjectPtr<UObject>> InAnimAssets);
+	void RetargetNonSkeletonAnimationHandler(USkeleton* OldSkeleton, USkeleton* NewSkeleton, bool bRemapReferencedAssets, bool bConvertSpaces, bool bDuplicateAssets, TArray<TWeakObjectPtr<UObject>> InAnimAssets, TWeakPtr<class IToolkitHost> EditWithinLevelEditor);
 };
\ No newline at end of file
diff --git a/Engine/Source/Editor/UnrealEd/Private/EditorAnimUtils.cpp b/Engine/Source/Editor/UnrealEd/Private/EditorAnimUtils.cpp
--- a/Engine/Source/Editor/UnrealEd/Private/EditorAnimUtils.cpp
+++ b/Engine/Source/Editor/UnrealEd/Private/EditorAnimUtils.cpp
@@ -169,18 +169,22 @@ namespace EditorAnimUtils
 			}
 		}
 
+
 		for(auto Iter = AnimSequencesToRetarget.CreateIterator(); Iter; ++Iter)
 		{
 			UAnimSequence* AssetToRetarget = (*Iter);
 
 			// Copy curve data from source asset, preserving data in the target if present.
-			FSmartNameMapping* OldNameMapping = OldSkeleton->SmartNames.GetContainer(USkeleton::AnimCurveMappingName);
-			FSmartNameMapping* NewNameMapping = NewSkeleton->SmartNames.GetContainer(USkeleton::AnimCurveMappingName);
-			AssetToRetarget->RawCurveData.UpdateLastObservedNames(OldNameMapping);
-
-			for(FFloatCurve& Curve : AssetToRetarget->RawCurveData.FloatCurves)
+			if (OldSkeleton)
 			{
-				NewNameMapping->AddName(Curve.LastObservedName, Curve.CurveUid);
+				FSmartNameMapping* OldNameMapping = OldSkeleton->SmartNames.GetContainer(USkeleton::AnimCurveMappingName);
+				FSmartNameMapping* NewNameMapping = NewSkeleton->SmartNames.GetContainer(USkeleton::AnimCurveMappingName);
+				AssetToRetarget->RawCurveData.UpdateLastObservedNames(OldNameMapping);
+
+				for(FFloatCurve& Curve : AssetToRetarget->RawCurveData.FloatCurves)
+				{
+					NewNameMapping->AddName(Curve.LastObservedName, Curve.CurveUid);
+				}
 			}
 
 			AssetToRetarget->ReplaceReferredAnimations(DuplicatedSequences);
