Prepare cpp files for header cleanup.

diff --git a/Engine/Plugins/2D/Paper2D/Source/Paper2D/Classes/PaperTerrainMaterial.h b/Engine/Plugins/2D/Paper2D/Source/Paper2D/Classes/PaperTerrainMaterial.h
--- a/Engine/Plugins/2D/Paper2D/Source/Paper2D/Classes/PaperTerrainMaterial.h
+++ b/Engine/Plugins/2D/Paper2D/Source/Paper2D/Classes/PaperTerrainMaterial.h
@@ -1,7 +1,7 @@
 // Copyright 1998-2014 Epic Games, Inc. All Rights Reserved.
 
 #pragma once
-
+#include "Engine/DataAsset.h"
 #include "PaperTerrainMaterial.generated.h"
 
 // Rule for a single section of a terrain material
diff --git a/Engine/Plugins/2D/Paper2D/Source/Paper2D/Classes/PaperTileMap.h b/Engine/Plugins/2D/Paper2D/Source/Paper2D/Classes/PaperTileMap.h
--- a/Engine/Plugins/2D/Paper2D/Source/Paper2D/Classes/PaperTileMap.h
+++ b/Engine/Plugins/2D/Paper2D/Source/Paper2D/Classes/PaperTileMap.h
@@ -2,6 +2,7 @@
 
 #pragma once
 
+#include "Engine/DataAsset.h"
 #include "Engine/EngineTypes.h"
 #include "PaperSprite.h"
 #include "PaperTileMap.generated.h"
diff --git a/Engine/Plugins/2D/Paper2D/Source/Paper2D/Classes/PaperTileSet.h b/Engine/Plugins/2D/Paper2D/Source/Paper2D/Classes/PaperTileSet.h
--- a/Engine/Plugins/2D/Paper2D/Source/Paper2D/Classes/PaperTileSet.h
+++ b/Engine/Plugins/2D/Paper2D/Source/Paper2D/Classes/PaperTileSet.h
@@ -2,6 +2,7 @@
 
 #pragma once
 
+#include "Engine/DataAsset.h"
 #include "Engine/EngineTypes.h"
 
 #include "PaperTileSet.generated.h"
diff --git a/Engine/Plugins/2D/Paper2D/Source/Paper2DEditor/Private/FlipbookEditor/FlipbookEditorViewportClient.cpp b/Engine/Plugins/2D/Paper2D/Source/Paper2DEditor/Private/FlipbookEditor/FlipbookEditorViewportClient.cpp
--- a/Engine/Plugins/2D/Paper2D/Source/Paper2DEditor/Private/FlipbookEditor/FlipbookEditorViewportClient.cpp
+++ b/Engine/Plugins/2D/Paper2D/Source/Paper2DEditor/Private/FlipbookEditor/FlipbookEditorViewportClient.cpp
@@ -6,6 +6,7 @@
 
 #include "PreviewScene.h"
 #include "ScopedTransaction.h"
+#include "CanvasTypes.h"
 
 #define LOCTEXT_NAMESPACE "FlipbookEditor"
 
diff --git a/Engine/Plugins/2D/Paper2D/Source/Paper2DEditor/Private/PaperEditorViewportClient.cpp b/Engine/Plugins/2D/Paper2D/Source/Paper2DEditor/Private/PaperEditorViewportClient.cpp
--- a/Engine/Plugins/2D/Paper2D/Source/Paper2DEditor/Private/PaperEditorViewportClient.cpp
+++ b/Engine/Plugins/2D/Paper2D/Source/Paper2DEditor/Private/PaperEditorViewportClient.cpp
@@ -3,6 +3,7 @@
 #include "Paper2DEditorPrivatePCH.h"
 #include "TileSetEditor.h"
 #include "PaperEditorViewportClient.h"
+#include "CanvasTypes.h"
 
 //////////////////////////////////////////////////////////////////////////
 // FPaperEditorViewportClient
diff --git a/Engine/Plugins/2D/Paper2D/Source/Paper2DEditor/Private/PaperFlipbookThumbnailRenderer.cpp b/Engine/Plugins/2D/Paper2D/Source/Paper2DEditor/Private/PaperFlipbookThumbnailRenderer.cpp
--- a/Engine/Plugins/2D/Paper2D/Source/Paper2DEditor/Private/PaperFlipbookThumbnailRenderer.cpp
+++ b/Engine/Plugins/2D/Paper2D/Source/Paper2DEditor/Private/PaperFlipbookThumbnailRenderer.cpp
@@ -1,6 +1,7 @@
 // Copyright 1998-2014 Epic Games, Inc. All Rights Reserved.
 
 #include "Paper2DEditorPrivatePCH.h"
+#include "CanvasItem.h"
 
 //////////////////////////////////////////////////////////////////////////
 // UPaperFlipbookThumbnailRenderer
diff --git a/Engine/Plugins/2D/Paper2D/Source/Paper2DEditor/Private/PaperSpriteActorFactory.cpp b/Engine/Plugins/2D/Paper2D/Source/Paper2DEditor/Private/PaperSpriteActorFactory.cpp
--- a/Engine/Plugins/2D/Paper2D/Source/Paper2DEditor/Private/PaperSpriteActorFactory.cpp
+++ b/Engine/Plugins/2D/Paper2D/Source/Paper2DEditor/Private/PaperSpriteActorFactory.cpp
@@ -2,7 +2,7 @@
 
 #include "Paper2DEditorPrivatePCH.h"
 #include "AssetData.h"
-
+#include "PhysicsEngine/BodySetup.h"
 //////////////////////////////////////////////////////////////////////////
 // UPaperSpriteActorFactory
 
diff --git a/Engine/Plugins/2D/Paper2D/Source/Paper2DEditor/Private/PaperSpriteAssetBroker.h b/Engine/Plugins/2D/Paper2D/Source/Paper2DEditor/Private/PaperSpriteAssetBroker.h
--- a/Engine/Plugins/2D/Paper2D/Source/Paper2DEditor/Private/PaperSpriteAssetBroker.h
+++ b/Engine/Plugins/2D/Paper2D/Source/Paper2DEditor/Private/PaperSpriteAssetBroker.h
@@ -3,6 +3,7 @@
 #pragma once
 
 #include "ComponentAssetBroker.h"
+#include "PhysicsEngine/BodySetup.h"
 
 //////////////////////////////////////////////////////////////////////////
 // FPaperSpriteAssetBroker
diff --git a/Engine/Plugins/2D/Paper2D/Source/Paper2DEditor/Private/PaperSpriteThumbnailRenderer.cpp b/Engine/Plugins/2D/Paper2D/Source/Paper2DEditor/Private/PaperSpriteThumbnailRenderer.cpp
--- a/Engine/Plugins/2D/Paper2D/Source/Paper2DEditor/Private/PaperSpriteThumbnailRenderer.cpp
+++ b/Engine/Plugins/2D/Paper2D/Source/Paper2DEditor/Private/PaperSpriteThumbnailRenderer.cpp
@@ -1,6 +1,7 @@
 // Copyright 1998-2014 Epic Games, Inc. All Rights Reserved.
 
 #include "Paper2DEditorPrivatePCH.h"
+#include "CanvasTypes.h"
 
 //////////////////////////////////////////////////////////////////////////
 // UPaperSpriteThumbnailRenderer
diff --git a/Engine/Plugins/2D/Paper2D/Source/Paper2DEditor/Private/PaperTileSetThumbnailRenderer.cpp b/Engine/Plugins/2D/Paper2D/Source/Paper2DEditor/Private/PaperTileSetThumbnailRenderer.cpp
--- a/Engine/Plugins/2D/Paper2D/Source/Paper2DEditor/Private/PaperTileSetThumbnailRenderer.cpp
+++ b/Engine/Plugins/2D/Paper2D/Source/Paper2DEditor/Private/PaperTileSetThumbnailRenderer.cpp
@@ -1,6 +1,7 @@
 // Copyright 1998-2014 Epic Games, Inc. All Rights Reserved.
 
 #include "Paper2DEditorPrivatePCH.h"
+#include "CanvasTypes.h"
 
 //////////////////////////////////////////////////////////////////////////
 // UPaperTileSetThumbnailRenderer
diff --git a/Engine/Plugins/2D/Paper2D/Source/Paper2DEditor/Private/SpriteEditor/SpriteDetailsCustomization.cpp b/Engine/Plugins/2D/Paper2D/Source/Paper2DEditor/Private/SpriteEditor/SpriteDetailsCustomization.cpp
--- a/Engine/Plugins/2D/Paper2D/Source/Paper2DEditor/Private/SpriteEditor/SpriteDetailsCustomization.cpp
+++ b/Engine/Plugins/2D/Paper2D/Source/Paper2DEditor/Private/SpriteEditor/SpriteDetailsCustomization.cpp
@@ -5,6 +5,7 @@
 
 #include "PhysicsEngine/BodySetup.h"
 #include "Editor/Documentation/Public/IDocumentation.h"
+#include "PhysicsEngine/PhysicsSettings.h"
 
 #define LOCTEXT_NAMESPACE "SpriteEditor"
 
diff --git a/Engine/Plugins/2D/Paper2D/Source/Paper2DEditor/Private/SpriteEditor/SpriteEditorViewportClient.cpp b/Engine/Plugins/2D/Paper2D/Source/Paper2DEditor/Private/SpriteEditor/SpriteEditorViewportClient.cpp
--- a/Engine/Plugins/2D/Paper2D/Source/Paper2DEditor/Private/SpriteEditor/SpriteEditorViewportClient.cpp
+++ b/Engine/Plugins/2D/Paper2D/Source/Paper2DEditor/Private/SpriteEditor/SpriteEditorViewportClient.cpp
@@ -11,6 +11,8 @@
 #include "AssetToolsModule.h"
 #include "ContentBrowserModule.h"
 #include "AssetRegistryModule.h"
+#include "CanvasTypes.h"
+#include "CanvasItem.h"
 
 #define LOCTEXT_NAMESPACE "SpriteEditor"
 
diff --git a/Engine/Plugins/2D/Paper2D/Source/Paper2DEditor/Private/TileMapEditing/EdModeTileMap.cpp b/Engine/Plugins/2D/Paper2D/Source/Paper2DEditor/Private/TileMapEditing/EdModeTileMap.cpp
--- a/Engine/Plugins/2D/Paper2D/Source/Paper2DEditor/Private/TileMapEditing/EdModeTileMap.cpp
+++ b/Engine/Plugins/2D/Paper2D/Source/Paper2DEditor/Private/TileMapEditing/EdModeTileMap.cpp
@@ -7,6 +7,9 @@
 #include "Toolkits/ToolkitManager.h"
 #include "../PaperEditorCommands.h"
 #include "ScopedTransaction.h"
+#include "CanvasItem.h"
+#include "CanvasTypes.h"
+#include "Engine/Selection.h"
 
 #define LOCTEXT_NAMESPACE "Paper2D"
 
diff --git a/Engine/Plugins/2D/Paper2D/Source/Paper2DEditor/Private/TileMapEditing/EdModeTileMap.h b/Engine/Plugins/2D/Paper2D/Source/Paper2DEditor/Private/TileMapEditing/EdModeTileMap.h
--- a/Engine/Plugins/2D/Paper2D/Source/Paper2DEditor/Private/TileMapEditing/EdModeTileMap.h
+++ b/Engine/Plugins/2D/Paper2D/Source/Paper2DEditor/Private/TileMapEditing/EdModeTileMap.h
@@ -1,6 +1,7 @@
 // Copyright 1998-2014 Epic Games, Inc. All Rights Reserved.
 
 #pragma once
+#include "PaperTileSet.h"
 
 namespace ETileMapEditorTool
 {
diff --git a/Engine/Plugins/2D/Paper2D/Source/Paper2DEditor/Private/TileMapEditing/TileMapEdModeToolkit.cpp b/Engine/Plugins/2D/Paper2D/Source/Paper2DEditor/Private/TileMapEditing/TileMapEdModeToolkit.cpp
--- a/Engine/Plugins/2D/Paper2D/Source/Paper2DEditor/Private/TileMapEditing/TileMapEdModeToolkit.cpp
+++ b/Engine/Plugins/2D/Paper2D/Source/Paper2DEditor/Private/TileMapEditing/TileMapEdModeToolkit.cpp
@@ -6,6 +6,7 @@
 #include "../TileSetEditor.h"
 #include "SContentReference.h"
 #include "PaperEditorCommands.h"
+#include "Engine/Selection.h"
 
 #define LOCTEXT_NAMESPACE "Paper2D"
 
diff --git a/Engine/Plugins/2D/Paper2D/Source/Paper2DEditor/Private/TileMapEditing/TileMapEditor.cpp b/Engine/Plugins/2D/Paper2D/Source/Paper2DEditor/Private/TileMapEditing/TileMapEditor.cpp
--- a/Engine/Plugins/2D/Paper2D/Source/Paper2DEditor/Private/TileMapEditing/TileMapEditor.cpp
+++ b/Engine/Plugins/2D/Paper2D/Source/Paper2DEditor/Private/TileMapEditing/TileMapEditor.cpp
@@ -14,6 +14,7 @@
 #include "Paper2DEditorModule.h"
 #include "STileMapEditorViewportToolbar.h"
 #include "SDockTab.h"
+#include "EdModeTileMap.h"
 
 #define LOCTEXT_NAMESPACE "TileMapEditor"
 
diff --git a/Engine/Plugins/2D/Paper2D/Source/Paper2DEditor/Private/TileMapEditing/TileMapEditorViewportClient.cpp b/Engine/Plugins/2D/Paper2D/Source/Paper2DEditor/Private/TileMapEditing/TileMapEditorViewportClient.cpp
--- a/Engine/Plugins/2D/Paper2D/Source/Paper2DEditor/Private/TileMapEditing/TileMapEditorViewportClient.cpp
+++ b/Engine/Plugins/2D/Paper2D/Source/Paper2DEditor/Private/TileMapEditing/TileMapEditorViewportClient.cpp
@@ -7,6 +7,7 @@
 #include "PreviewScene.h"
 #include "ScopedTransaction.h"
 #include "Runtime/Engine/Public/ComponentReregisterContext.h"
+#include "CanvasTypes.h"
 
 #define LOCTEXT_NAMESPACE "TileMapEditor"
 
diff --git a/Engine/Plugins/2D/Paper2D/Source/Paper2DEditor/Private/TileSetEditor.cpp b/Engine/Plugins/2D/Paper2D/Source/Paper2DEditor/Private/TileSetEditor.cpp
--- a/Engine/Plugins/2D/Paper2D/Source/Paper2DEditor/Private/TileSetEditor.cpp
+++ b/Engine/Plugins/2D/Paper2D/Source/Paper2DEditor/Private/TileSetEditor.cpp
@@ -8,6 +8,8 @@
 #include "TileMapEditing/EdModeTileMap.h"
 #include "GraphEditor.h"
 #include "SPaperEditorViewport.h"
+#include "CanvasTypes.h"
+#include "CanvasItem.h"
 
 #define LOCTEXT_NAMESPACE "TileSetEditor"
 
diff --git a/Engine/Plugins/Editor/SpeedTreeImporter/Source/SpeedTreeImporter/Private/ReimportSpeedTreeFactory.cpp b/Engine/Plugins/Editor/SpeedTreeImporter/Source/SpeedTreeImporter/Private/ReimportSpeedTreeFactory.cpp
--- a/Engine/Plugins/Editor/SpeedTreeImporter/Source/SpeedTreeImporter/Private/ReimportSpeedTreeFactory.cpp
+++ b/Engine/Plugins/Editor/SpeedTreeImporter/Source/SpeedTreeImporter/Private/ReimportSpeedTreeFactory.cpp
@@ -2,6 +2,7 @@
 
 #include "SpeedTreeImporterPrivatePCH.h"
 #include "ReimportSpeedTreeFactory.h"
+#include "Engine/StaticMesh.h"
 
 
 
diff --git a/Engine/Plugins/Editor/SpeedTreeImporter/Source/SpeedTreeImporter/Private/SpeedTreeImportFactory.cpp b/Engine/Plugins/Editor/SpeedTreeImporter/Source/SpeedTreeImporter/Private/SpeedTreeImportFactory.cpp
--- a/Engine/Plugins/Editor/SpeedTreeImporter/Source/SpeedTreeImporter/Private/SpeedTreeImportFactory.cpp
+++ b/Engine/Plugins/Editor/SpeedTreeImporter/Source/SpeedTreeImporter/Private/SpeedTreeImportFactory.cpp
@@ -43,6 +43,8 @@
 #endif
 
 #endif // WITH_SPEEDTREE
+#include "Engine/StaticMesh.h"
+#include "PhysicsEngine/BodySetup.h"
 
 #define LOCTEXT_NAMESPACE "SpeedTreeImportFactory"
 
diff --git a/Engine/Plugins/Media/MediaPlayerEditor/Source/MediaPlayerEditor/Private/MediaPlayerEditorToolkit.h b/Engine/Plugins/Media/MediaPlayerEditor/Source/MediaPlayerEditor/Private/MediaPlayerEditorToolkit.h
--- a/Engine/Plugins/Media/MediaPlayerEditor/Source/MediaPlayerEditor/Private/MediaPlayerEditorToolkit.h
+++ b/Engine/Plugins/Media/MediaPlayerEditor/Source/MediaPlayerEditor/Private/MediaPlayerEditorToolkit.h
@@ -3,6 +3,7 @@
 #pragma once
 
 #include "Toolkits/AssetEditorToolkit.h"
+#include "EditorUndoClient.h"
 
 
 /**
diff --git a/Engine/Plugins/Runtime/SoundMod/Source/SoundMod/Classes/SoundMod.h b/Engine/Plugins/Runtime/SoundMod/Source/SoundMod/Classes/SoundMod.h
--- a/Engine/Plugins/Runtime/SoundMod/Source/SoundMod/Classes/SoundMod.h
+++ b/Engine/Plugins/Runtime/SoundMod/Source/SoundMod/Classes/SoundMod.h
@@ -6,7 +6,9 @@
  * A sound module file
  */
 
+#include "Sound/SoundBase.h"
 #include "SoundMod.generated.h"
+struct FActiveSound;
 
 UCLASS(hidecategories=Object, MinimalAPI, BlueprintType)
 class USoundMod : public USoundBase
diff --git a/Engine/Plugins/Runtime/SoundMod/Source/SoundModImporter/Private/AssetTypeActions_SoundMod.cpp b/Engine/Plugins/Runtime/SoundMod/Source/SoundModImporter/Private/AssetTypeActions_SoundMod.cpp
--- a/Engine/Plugins/Runtime/SoundMod/Source/SoundModImporter/Private/AssetTypeActions_SoundMod.cpp
+++ b/Engine/Plugins/Runtime/SoundMod/Source/SoundModImporter/Private/AssetTypeActions_SoundMod.cpp
@@ -2,6 +2,7 @@
 
 #include "SoundModImporterPrivatePCH.h"
 #include "AssetTypeActions_SoundMod.h"
+#include "Components/AudioComponent.h"
 
 // Copyright 1998-2014 Epic Games, Inc. All Rights Reserved.
 
diff --git a/Engine/Plugins/Runtime/SoundMod/Source/SoundModImporter/Private/SoundModImporterFactory.cpp b/Engine/Plugins/Runtime/SoundMod/Source/SoundModImporter/Private/SoundModImporterFactory.cpp
--- a/Engine/Plugins/Runtime/SoundMod/Source/SoundModImporter/Private/SoundModImporterFactory.cpp
+++ b/Engine/Plugins/Runtime/SoundMod/Source/SoundModImporter/Private/SoundModImporterFactory.cpp
@@ -6,6 +6,7 @@
 #include "PackageTools.h"
 #include "SoundDefinitions.h"
 #include "xmp.h"
+#include "Components/AudioComponent.h"
 
 //////////////////////////////////////////////////////////////////////////
 // USoundModImporterFactory
diff --git a/Engine/Source/Developer/AssetTools/Private/AssetTypeActions/AssetTypeActions_AimOffset.h b/Engine/Source/Developer/AssetTools/Private/AssetTypeActions/AssetTypeActions_AimOffset.h
--- a/Engine/Source/Developer/AssetTools/Private/AssetTypeActions/AssetTypeActions_AimOffset.h
+++ b/Engine/Source/Developer/AssetTools/Private/AssetTypeActions/AssetTypeActions_AimOffset.h
@@ -1,6 +1,8 @@
 // Copyright 1998-2014 Epic Games, Inc. All Rights Reserved.
 
 #pragma once
+#include "Animation/AimOffsetBlendSpace.h"
+#include "AssetTypeActions_BlendSpace.h"
 
 class FAssetTypeActions_AimOffset : public FAssetTypeActions_BlendSpace
 {
diff --git a/Engine/Source/Developer/AssetTools/Private/AssetTypeActions/AssetTypeActions_AimOffset1D.h b/Engine/Source/Developer/AssetTools/Private/AssetTypeActions/AssetTypeActions_AimOffset1D.h
--- a/Engine/Source/Developer/AssetTools/Private/AssetTypeActions/AssetTypeActions_AimOffset1D.h
+++ b/Engine/Source/Developer/AssetTools/Private/AssetTypeActions/AssetTypeActions_AimOffset1D.h
@@ -1,6 +1,8 @@
 // Copyright 1998-2014 Epic Games, Inc. All Rights Reserved.
 
 #pragma once
+#include "Animation/AimOffsetBlendSpace1D.h"
+#include "AssetTypeActions_BlendSpace1D.h"
 
 class FAssetTypeActions_AimOffset1D : public FAssetTypeActions_BlendSpace1D
 {
diff --git a/Engine/Source/Developer/AssetTools/Private/AssetTypeActions/AssetTypeActions_AnimComposite.h b/Engine/Source/Developer/AssetTools/Private/AssetTypeActions/AssetTypeActions_AnimComposite.h
--- a/Engine/Source/Developer/AssetTools/Private/AssetTypeActions/AssetTypeActions_AnimComposite.h
+++ b/Engine/Source/Developer/AssetTools/Private/AssetTypeActions/AssetTypeActions_AnimComposite.h
@@ -1,6 +1,8 @@
 // Copyright 1998-2014 Epic Games, Inc. All Rights Reserved.
 
 #pragma once
+#include "Animation/AnimComposite.h"
+#include "AssetTypeActions_AnimationAsset.h"
 
 class FAssetTypeActions_AnimComposite : public FAssetTypeActions_AnimationAsset
 {
diff --git a/Engine/Source/Developer/AssetTools/Private/AssetTypeActions/AssetTypeActions_BlendSpace.h b/Engine/Source/Developer/AssetTools/Private/AssetTypeActions/AssetTypeActions_BlendSpace.h
--- a/Engine/Source/Developer/AssetTools/Private/AssetTypeActions/AssetTypeActions_BlendSpace.h
+++ b/Engine/Source/Developer/AssetTools/Private/AssetTypeActions/AssetTypeActions_BlendSpace.h
@@ -1,6 +1,8 @@
 // Copyright 1998-2014 Epic Games, Inc. All Rights Reserved.
 
 #pragma once
+#include "Animation/BlendSpace.h"
+#include "AssetTypeActions_AnimationAsset.h"
 
 class FAssetTypeActions_BlendSpace : public FAssetTypeActions_AnimationAsset
 {
diff --git a/Engine/Source/Developer/AssetTools/Private/AssetTypeActions/AssetTypeActions_BlendSpace1D.h b/Engine/Source/Developer/AssetTools/Private/AssetTypeActions/AssetTypeActions_BlendSpace1D.h
--- a/Engine/Source/Developer/AssetTools/Private/AssetTypeActions/AssetTypeActions_BlendSpace1D.h
+++ b/Engine/Source/Developer/AssetTools/Private/AssetTypeActions/AssetTypeActions_BlendSpace1D.h
@@ -1,6 +1,8 @@
 // Copyright 1998-2014 Epic Games, Inc. All Rights Reserved.
 
 #pragma once
+#include "Animation/BlendSpace1D.h"
+#include "AssetTypeActions_AnimationAsset.h"
 
 class FAssetTypeActions_BlendSpace1D : public FAssetTypeActions_AnimationAsset
 {
diff --git a/Engine/Source/Developer/AssetTools/Private/AssetTypeActions/AssetTypeActions_Blueprint.cpp b/Engine/Source/Developer/AssetTools/Private/AssetTypeActions/AssetTypeActions_Blueprint.cpp
--- a/Engine/Source/Developer/AssetTools/Private/AssetTypeActions/AssetTypeActions_Blueprint.cpp
+++ b/Engine/Source/Developer/AssetTools/Private/AssetTypeActions/AssetTypeActions_Blueprint.cpp
@@ -10,6 +10,7 @@
 #include "SBlueprintDiff.h"
 #include "ISourceControlModule.h"
 #include "MessageLog.h"
+#include "Engine/BlueprintGeneratedClass.h"
 
 #define LOCTEXT_NAMESPACE "AssetTypeActions"
 
diff --git a/Engine/Source/Developer/AssetTools/Private/AssetTypeActions/AssetTypeActions_CameraAnim.cpp b/Engine/Source/Developer/AssetTools/Private/AssetTypeActions/AssetTypeActions_CameraAnim.cpp
--- a/Engine/Source/Developer/AssetTools/Private/AssetTypeActions/AssetTypeActions_CameraAnim.cpp
+++ b/Engine/Source/Developer/AssetTools/Private/AssetTypeActions/AssetTypeActions_CameraAnim.cpp
@@ -8,6 +8,9 @@
 #include "Matinee/InterpTrackInst.h"
 #include "Matinee/InterpTrackAnimControl.h"
 #include "Matinee/MatineeActorCameraAnim.h"
+#include "Camera/CameraAnim.h"
+#include "Camera/CameraActor.h"
+#include "GameFramework/Pawn.h"
 
 DEFINE_LOG_CATEGORY_STATIC(LogCameraAnim, Log, All);
 
diff --git a/Engine/Source/Developer/AssetTools/Private/AssetTypeActions/AssetTypeActions_CameraAnim.h b/Engine/Source/Developer/AssetTools/Private/AssetTypeActions/AssetTypeActions_CameraAnim.h
--- a/Engine/Source/Developer/AssetTools/Private/AssetTypeActions/AssetTypeActions_CameraAnim.h
+++ b/Engine/Source/Developer/AssetTools/Private/AssetTypeActions/AssetTypeActions_CameraAnim.h
@@ -1,6 +1,8 @@
 // Copyright 1998-2014 Epic Games, Inc. All Rights Reserved.
 
 #pragma once
+#include "Camera/CameraAnim.h"
+#include "AssetTypeActions_Base.h"
 
 class FAssetTypeActions_CameraAnim : public FAssetTypeActions_Base
 {
diff --git a/Engine/Source/Developer/AssetTools/Private/AssetTypeActions/AssetTypeActions_CurveFloat.h b/Engine/Source/Developer/AssetTools/Private/AssetTypeActions/AssetTypeActions_CurveFloat.h
--- a/Engine/Source/Developer/AssetTools/Private/AssetTypeActions/AssetTypeActions_CurveFloat.h
+++ b/Engine/Source/Developer/AssetTools/Private/AssetTypeActions/AssetTypeActions_CurveFloat.h
@@ -1,6 +1,8 @@
 // Copyright 1998-2014 Epic Games, Inc. All Rights Reserved.
 
 #pragma once
+#include "Curves/CurveFloat.h"
+#include "AssetTypeActions_Curve.h"
 
 class FAssetTypeActions_CurveFloat : public FAssetTypeActions_Curve
 {
diff --git a/Engine/Source/Developer/AssetTools/Private/AssetTypeActions/AssetTypeActions_CurveLinearColor.h b/Engine/Source/Developer/AssetTools/Private/AssetTypeActions/AssetTypeActions_CurveLinearColor.h
--- a/Engine/Source/Developer/AssetTools/Private/AssetTypeActions/AssetTypeActions_CurveLinearColor.h
+++ b/Engine/Source/Developer/AssetTools/Private/AssetTypeActions/AssetTypeActions_CurveLinearColor.h
@@ -1,6 +1,8 @@
 // Copyright 1998-2014 Epic Games, Inc. All Rights Reserved.
 
 #pragma once
+#include "Curves/CurveLinearColor.h"
+#include "AssetTypeActions_Curve.h"
 
 class FAssetTypeActions_CurveLinearColor : public FAssetTypeActions_Curve
 {
diff --git a/Engine/Source/Developer/AssetTools/Private/AssetTypeActions/AssetTypeActions_CurveTable.cpp b/Engine/Source/Developer/AssetTools/Private/AssetTypeActions/AssetTypeActions_CurveTable.cpp
--- a/Engine/Source/Developer/AssetTools/Private/AssetTypeActions/AssetTypeActions_CurveTable.cpp
+++ b/Engine/Source/Developer/AssetTools/Private/AssetTypeActions/AssetTypeActions_CurveTable.cpp
@@ -4,6 +4,7 @@
 
 #include "Editor/CurveTableEditor/Public/CurveTableEditorModule.h"
 #include "Editor/CurveTableEditor/Public/ICurveTableEditor.h"
+#include "Engine/CurveTable.h"
 
 #define LOCTEXT_NAMESPACE "AssetTypeActions"
 
diff --git a/Engine/Source/Developer/AssetTools/Private/AssetTypeActions/AssetTypeActions_CurveTable.h b/Engine/Source/Developer/AssetTools/Private/AssetTypeActions/AssetTypeActions_CurveTable.h
--- a/Engine/Source/Developer/AssetTools/Private/AssetTypeActions/AssetTypeActions_CurveTable.h
+++ b/Engine/Source/Developer/AssetTools/Private/AssetTypeActions/AssetTypeActions_CurveTable.h
@@ -1,6 +1,8 @@
 // Copyright 1998-2014 Epic Games, Inc. All Rights Reserved.
 
 #pragma once
+#include "Engine/CurveTable.h"
+#include "AssetTypeActions_CSVAssetBase.h"
 
 class FAssetTypeActions_CurveTable : public FAssetTypeActions_CSVAssetBase
 {
diff --git a/Engine/Source/Developer/AssetTools/Private/AssetTypeActions/AssetTypeActions_CurveVector.h b/Engine/Source/Developer/AssetTools/Private/AssetTypeActions/AssetTypeActions_CurveVector.h
--- a/Engine/Source/Developer/AssetTools/Private/AssetTypeActions/AssetTypeActions_CurveVector.h
+++ b/Engine/Source/Developer/AssetTools/Private/AssetTypeActions/AssetTypeActions_CurveVector.h
@@ -1,6 +1,8 @@
 // Copyright 1998-2014 Epic Games, Inc. All Rights Reserved.
 
 #pragma once
+#include "Curves/CurveVector.h"
+#include "AssetTypeActions_Curve.h"
 
 class FAssetTypeActions_CurveVector : public FAssetTypeActions_Curve
 {
diff --git a/Engine/Source/Developer/AssetTools/Private/AssetTypeActions/AssetTypeActions_DataAsset.h b/Engine/Source/Developer/AssetTools/Private/AssetTypeActions/AssetTypeActions_DataAsset.h
--- a/Engine/Source/Developer/AssetTools/Private/AssetTypeActions/AssetTypeActions_DataAsset.h
+++ b/Engine/Source/Developer/AssetTools/Private/AssetTypeActions/AssetTypeActions_DataAsset.h
@@ -1,6 +1,9 @@
 // Copyright 1998-2014 Epic Games, Inc. All Rights Reserved.
 
 #pragma once
+#include "Engine/DataAsset.h"
+#include "IAssetTypeActions.h"
+#include "AssetTypeActions_Base.h"
 
 class FAssetTypeActions_DataAsset : public FAssetTypeActions_Base
 {
diff --git a/Engine/Source/Developer/AssetTools/Private/AssetTypeActions/AssetTypeActions_DataTable.cpp b/Engine/Source/Developer/AssetTools/Private/AssetTypeActions/AssetTypeActions_DataTable.cpp
--- a/Engine/Source/Developer/AssetTools/Private/AssetTypeActions/AssetTypeActions_DataTable.cpp
+++ b/Engine/Source/Developer/AssetTools/Private/AssetTypeActions/AssetTypeActions_DataTable.cpp
@@ -4,6 +4,7 @@
 
 #include "Editor/DataTableEditor/Public/DataTableEditorModule.h"
 #include "Editor/DataTableEditor/Public/IDataTableEditor.h"
+#include "Engine/DataTable.h"
 
 #define LOCTEXT_NAMESPACE "AssetTypeActions"
 
diff --git a/Engine/Source/Developer/AssetTools/Private/AssetTypeActions/AssetTypeActions_DataTable.h b/Engine/Source/Developer/AssetTools/Private/AssetTypeActions/AssetTypeActions_DataTable.h
--- a/Engine/Source/Developer/AssetTools/Private/AssetTypeActions/AssetTypeActions_DataTable.h
+++ b/Engine/Source/Developer/AssetTools/Private/AssetTypeActions/AssetTypeActions_DataTable.h
@@ -1,6 +1,8 @@
 // Copyright 1998-2014 Epic Games, Inc. All Rights Reserved.
 
 #pragma once
+#include "Engine/DataTable.h"
+#include "AssetTypeActions_CSVAssetBase.h"
 
 class FAssetTypeActions_DataTable : public FAssetTypeActions_CSVAssetBase
 {
diff --git a/Engine/Source/Developer/AssetTools/Private/AssetTypeActions/AssetTypeActions_DialogueVoice.h b/Engine/Source/Developer/AssetTools/Private/AssetTypeActions/AssetTypeActions_DialogueVoice.h
--- a/Engine/Source/Developer/AssetTools/Private/AssetTypeActions/AssetTypeActions_DialogueVoice.h
+++ b/Engine/Source/Developer/AssetTools/Private/AssetTypeActions/AssetTypeActions_DialogueVoice.h
@@ -1,8 +1,9 @@
 // Copyright 1998-2014 Epic Games, Inc. All Rights Reserved.
 
 #pragma once
-
+#include "Sound/DialogueVoice.h"
 #include "SoundDefinitions.h"
+#include "AssetTypeActions_Base.h"
 
 class FAssetTypeActions_DialogueVoice : public FAssetTypeActions_Base
 {
diff --git a/Engine/Source/Developer/AssetTools/Private/AssetTypeActions/AssetTypeActions_DialogueWave.h b/Engine/Source/Developer/AssetTools/Private/AssetTypeActions/AssetTypeActions_DialogueWave.h
--- a/Engine/Source/Developer/AssetTools/Private/AssetTypeActions/AssetTypeActions_DialogueWave.h
+++ b/Engine/Source/Developer/AssetTools/Private/AssetTypeActions/AssetTypeActions_DialogueWave.h
@@ -1,8 +1,9 @@
 // Copyright 1998-2014 Epic Games, Inc. All Rights Reserved.
 
 #pragma once
-
+#include "Sound/DialogueWave.h"
 #include "SoundDefinitions.h"
+#include "AssetTypeActions_Base.h"
 
 class FAssetTypeActions_DialogueWave : public FAssetTypeActions_Base
 {
diff --git a/Engine/Source/Developer/AssetTools/Private/AssetTypeActions/AssetTypeActions_Enum.h b/Engine/Source/Developer/AssetTools/Private/AssetTypeActions/AssetTypeActions_Enum.h
--- a/Engine/Source/Developer/AssetTools/Private/AssetTypeActions/AssetTypeActions_Enum.h
+++ b/Engine/Source/Developer/AssetTools/Private/AssetTypeActions/AssetTypeActions_Enum.h
@@ -1,6 +1,8 @@
 // Copyright 1998-2014 Epic Games, Inc. All Rights Reserved.
 
 #pragma once
+#include "Engine/UserDefinedEnum.h"
+#include "AssetTypeActions_Base.h"
 
 class FAssetTypeActions_Enum : public FAssetTypeActions_Base
 {
diff --git a/Engine/Source/Developer/AssetTools/Private/AssetTypeActions/AssetTypeActions_Font.h b/Engine/Source/Developer/AssetTools/Private/AssetTypeActions/AssetTypeActions_Font.h
--- a/Engine/Source/Developer/AssetTools/Private/AssetTypeActions/AssetTypeActions_Font.h
+++ b/Engine/Source/Developer/AssetTools/Private/AssetTypeActions/AssetTypeActions_Font.h
@@ -1,6 +1,8 @@
 // Copyright 1998-2014 Epic Games, Inc. All Rights Reserved.
 
 #pragma once
+#include "Engine/Font.h"
+#include "AssetTypeActions_Base.h"
 
 class FAssetTypeActions_Font : public FAssetTypeActions_Base
 {
diff --git a/Engine/Source/Developer/AssetTools/Private/AssetTypeActions/AssetTypeActions_ForceFeedbackEffect.cpp b/Engine/Source/Developer/AssetTools/Private/AssetTypeActions/AssetTypeActions_ForceFeedbackEffect.cpp
--- a/Engine/Source/Developer/AssetTools/Private/AssetTypeActions/AssetTypeActions_ForceFeedbackEffect.cpp
+++ b/Engine/Source/Developer/AssetTools/Private/AssetTypeActions/AssetTypeActions_ForceFeedbackEffect.cpp
@@ -1,6 +1,7 @@
 // Copyright 1998-2014 Epic Games, Inc. All Rights Reserved.
 
 #include "AssetToolsPrivatePCH.h"
+#include "GameFramework/ForceFeedbackEffect.h"
 
 UClass* FAssetTypeActions_ForceFeedbackEffect::GetSupportedClass() const
 {
diff --git a/Engine/Source/Developer/AssetTools/Private/AssetTypeActions/AssetTypeActions_MaterialInstanceConstant.cpp b/Engine/Source/Developer/AssetTools/Private/AssetTypeActions/AssetTypeActions_MaterialInstanceConstant.cpp
--- a/Engine/Source/Developer/AssetTools/Private/AssetTypeActions/AssetTypeActions_MaterialInstanceConstant.cpp
+++ b/Engine/Source/Developer/AssetTools/Private/AssetTypeActions/AssetTypeActions_MaterialInstanceConstant.cpp
@@ -2,6 +2,7 @@
 
 #include "AssetToolsPrivatePCH.h"
 #include "Editor/MaterialEditor/Public/MaterialEditorModule.h"
+#include "Materials/MaterialInstanceConstant.h"
 
 #define LOCTEXT_NAMESPACE "AssetTypeActions"
 
diff --git a/Engine/Source/Developer/AssetTools/Private/AssetTypeActions/AssetTypeActions_MaterialInstanceConstant.h b/Engine/Source/Developer/AssetTools/Private/AssetTypeActions/AssetTypeActions_MaterialInstanceConstant.h
--- a/Engine/Source/Developer/AssetTools/Private/AssetTypeActions/AssetTypeActions_MaterialInstanceConstant.h
+++ b/Engine/Source/Developer/AssetTools/Private/AssetTypeActions/AssetTypeActions_MaterialInstanceConstant.h
@@ -1,6 +1,8 @@
 // Copyright 1998-2014 Epic Games, Inc. All Rights Reserved.
 
 #pragma once
+#include "Materials/MaterialInstanceConstant.h"
+#include "AssetTypeActions_MaterialInterface.h"
 
 class FAssetTypeActions_MaterialInstanceConstant : public FAssetTypeActions_MaterialInterface
 {
diff --git a/Engine/Source/Developer/AssetTools/Private/AssetTypeActions/AssetTypeActions_MaterialInterface.cpp b/Engine/Source/Developer/AssetTools/Private/AssetTypeActions/AssetTypeActions_MaterialInterface.cpp
--- a/Engine/Source/Developer/AssetTools/Private/AssetTypeActions/AssetTypeActions_MaterialInterface.cpp
+++ b/Engine/Source/Developer/AssetTools/Private/AssetTypeActions/AssetTypeActions_MaterialInterface.cpp
@@ -5,6 +5,7 @@
 #include "AssetRegistryModule.h"
 #include "ContentBrowserModule.h"
 #include "AssetToolsModule.h"
+#include "Materials/MaterialInstanceConstant.h"
 
 #define LOCTEXT_NAMESPACE "AssetTypeActions"
 
diff --git a/Engine/Source/Developer/AssetTools/Private/AssetTypeActions/AssetTypeActions_MaterialParameterCollection.h b/Engine/Source/Developer/AssetTools/Private/AssetTypeActions/AssetTypeActions_MaterialParameterCollection.h
--- a/Engine/Source/Developer/AssetTools/Private/AssetTypeActions/AssetTypeActions_MaterialParameterCollection.h
+++ b/Engine/Source/Developer/AssetTools/Private/AssetTypeActions/AssetTypeActions_MaterialParameterCollection.h
@@ -1,6 +1,8 @@
 // Copyright 1998-2014 Epic Games, Inc. All Rights Reserved.
 
 #pragma once
+#include "Materials/MaterialParameterCollection.h"
+#include "AssetTypeActions_Base.h"
 
 class FAssetTypeActions_MaterialParameterCollection : public FAssetTypeActions_Base
 {
diff --git a/Engine/Source/Developer/AssetTools/Private/AssetTypeActions/AssetTypeActions_MorphTarget.cpp b/Engine/Source/Developer/AssetTools/Private/AssetTypeActions/AssetTypeActions_MorphTarget.cpp
--- a/Engine/Source/Developer/AssetTools/Private/AssetTypeActions/AssetTypeActions_MorphTarget.cpp
+++ b/Engine/Source/Developer/AssetTools/Private/AssetTypeActions/AssetTypeActions_MorphTarget.cpp
@@ -4,6 +4,7 @@
 #include "Toolkits/AssetEditorManager.h"
 #include "Editor/Persona/Public/PersonaModule.h"
 #include "ObjectTools.h"
+#include "Animation/VertexAnim/MorphTarget.h"
 
 DEFINE_LOG_CATEGORY_STATIC(LogMorphTarget, Warning, All);
 
diff --git a/Engine/Source/Developer/AssetTools/Private/AssetTypeActions/AssetTypeActions_MorphTarget.h b/Engine/Source/Developer/AssetTools/Private/AssetTypeActions/AssetTypeActions_MorphTarget.h
--- a/Engine/Source/Developer/AssetTools/Private/AssetTypeActions/AssetTypeActions_MorphTarget.h
+++ b/Engine/Source/Developer/AssetTools/Private/AssetTypeActions/AssetTypeActions_MorphTarget.h
@@ -1,6 +1,8 @@
 // Copyright 1998-2014 Epic Games, Inc. All Rights Reserved.
 
 #pragma once
+#include "Animation/VertexAnim/MorphTarget.h"
+#include "AssetTypeActions_Base.h"
 
 class FAssetTypeActions_MorphTarget : public FAssetTypeActions_Base
 {
diff --git a/Engine/Source/Developer/AssetTools/Private/AssetTypeActions/AssetTypeActions_ObjectLibrary.h b/Engine/Source/Developer/AssetTools/Private/AssetTypeActions/AssetTypeActions_ObjectLibrary.h
--- a/Engine/Source/Developer/AssetTools/Private/AssetTypeActions/AssetTypeActions_ObjectLibrary.h
+++ b/Engine/Source/Developer/AssetTools/Private/AssetTypeActions/AssetTypeActions_ObjectLibrary.h
@@ -1,6 +1,8 @@
 // Copyright 1998-2014 Epic Games, Inc. All Rights Reserved.
 
 #pragma once
+#include "Engine/ObjectLibrary.h"
+#include "AssetTypeActions_Base.h"
 
 class FAssetTypeActions_ObjectLibrary : public FAssetTypeActions_Base
 {
diff --git a/Engine/Source/Developer/AssetTools/Private/AssetTypeActions/AssetTypeActions_PhysicalMaterial.cpp b/Engine/Source/Developer/AssetTools/Private/AssetTypeActions/AssetTypeActions_PhysicalMaterial.cpp
--- a/Engine/Source/Developer/AssetTools/Private/AssetTypeActions/AssetTypeActions_PhysicalMaterial.cpp
+++ b/Engine/Source/Developer/AssetTools/Private/AssetTypeActions/AssetTypeActions_PhysicalMaterial.cpp
@@ -1,6 +1,7 @@
 // Copyright 1998-2014 Epic Games, Inc. All Rights Reserved.
 
 #include "AssetToolsPrivatePCH.h"
+#include "PhysicalMaterials/PhysicalMaterial.h"
 
 UClass* FAssetTypeActions_PhysicalMaterial::GetSupportedClass() const
 {
diff --git a/Engine/Source/Developer/AssetTools/Private/AssetTypeActions/AssetTypeActions_PhysicsAsset.cpp b/Engine/Source/Developer/AssetTools/Private/AssetTypeActions/AssetTypeActions_PhysicsAsset.cpp
--- a/Engine/Source/Developer/AssetTools/Private/AssetTypeActions/AssetTypeActions_PhysicsAsset.cpp
+++ b/Engine/Source/Developer/AssetTools/Private/AssetTypeActions/AssetTypeActions_PhysicsAsset.cpp
@@ -3,6 +3,7 @@
 #include "AssetToolsPrivatePCH.h"
 #include "Toolkits/AssetEditorManager.h"
 #include "Editor/PhAT/Public/PhATModule.h"
+#include "PhysicsEngine/PhysicsAsset.h"
 
 #define LOCTEXT_NAMESPACE "AssetTypeActions"
 
diff --git a/Engine/Source/Developer/AssetTools/Private/AssetTypeActions/AssetTypeActions_ReverbEffect.cpp b/Engine/Source/Developer/AssetTools/Private/AssetTypeActions/AssetTypeActions_ReverbEffect.cpp
--- a/Engine/Source/Developer/AssetTools/Private/AssetTypeActions/AssetTypeActions_ReverbEffect.cpp
+++ b/Engine/Source/Developer/AssetTools/Private/AssetTypeActions/AssetTypeActions_ReverbEffect.cpp
@@ -2,6 +2,7 @@
 
 #include "AssetToolsPrivatePCH.h"
 #include "SoundDefinitions.h"
+#include "Sound/ReverbEffect.h"
 
 UClass* FAssetTypeActions_ReverbEffect::GetSupportedClass() const
 {
diff --git a/Engine/Source/Developer/AssetTools/Private/AssetTypeActions/AssetTypeActions_SkeletalMesh.cpp b/Engine/Source/Developer/AssetTools/Private/AssetTypeActions/AssetTypeActions_SkeletalMesh.cpp
--- a/Engine/Source/Developer/AssetTools/Private/AssetTypeActions/AssetTypeActions_SkeletalMesh.cpp
+++ b/Engine/Source/Developer/AssetTools/Private/AssetTypeActions/AssetTypeActions_SkeletalMesh.cpp
@@ -14,6 +14,7 @@
 #include "AssetToolsModule.h"
 #include "Editor/UnrealEd/Public/FbxMeshUtils.h"
 #include "Editor/UnrealEd/Public/AssetNotifications.h"
+#include "PhysicsEngine/PhysicsAsset.h"
 
 #define LOCTEXT_NAMESPACE "AssetTypeActions"
 
diff --git a/Engine/Source/Developer/AssetTools/Private/AssetTypeActions/AssetTypeActions_SoundBase.cpp b/Engine/Source/Developer/AssetTools/Private/AssetTypeActions/AssetTypeActions_SoundBase.cpp
--- a/Engine/Source/Developer/AssetTools/Private/AssetTypeActions/AssetTypeActions_SoundBase.cpp
+++ b/Engine/Source/Developer/AssetTools/Private/AssetTypeActions/AssetTypeActions_SoundBase.cpp
@@ -4,6 +4,8 @@
 #include "Editor/UnrealEd/Public/Dialogs/DlgSoundWaveOptions.h"
 #include "Editor/SoundCueEditor/Public/SoundCueEditorModule.h"
 #include "SoundDefinitions.h"
+#include "Sound/SoundBase.h"
+#include "Components/AudioComponent.h"
 
 #define LOCTEXT_NAMESPACE "AssetTypeActions"
 
diff --git a/Engine/Source/Developer/AssetTools/Private/AssetTypeActions/AssetTypeActions_SoundCue.cpp b/Engine/Source/Developer/AssetTools/Private/AssetTypeActions/AssetTypeActions_SoundCue.cpp
--- a/Engine/Source/Developer/AssetTools/Private/AssetTypeActions/AssetTypeActions_SoundCue.cpp
+++ b/Engine/Source/Developer/AssetTools/Private/AssetTypeActions/AssetTypeActions_SoundCue.cpp
@@ -5,6 +5,7 @@
 #include "Editor/UnrealEd/Public/Dialogs/DlgSoundWaveOptions.h"
 #include "Editor/SoundCueEditor/Public/SoundCueEditorModule.h"
 #include "SoundDefinitions.h"
+#include "Sound/SoundCue.h"
 
 #define LOCTEXT_NAMESPACE "AssetTypeActions"
 
diff --git a/Engine/Source/Developer/AssetTools/Private/AssetTypeActions/AssetTypeActions_SoundMix.cpp b/Engine/Source/Developer/AssetTools/Private/AssetTypeActions/AssetTypeActions_SoundMix.cpp
--- a/Engine/Source/Developer/AssetTools/Private/AssetTypeActions/AssetTypeActions_SoundMix.cpp
+++ b/Engine/Source/Developer/AssetTools/Private/AssetTypeActions/AssetTypeActions_SoundMix.cpp
@@ -2,6 +2,7 @@
 
 #include "AssetToolsPrivatePCH.h"
 #include "SoundDefinitions.h"
+#include "Sound/SoundMix.h"
 
 UClass* FAssetTypeActions_SoundMix::GetSupportedClass() const
 {
diff --git a/Engine/Source/Developer/AssetTools/Private/AssetTypeActions/AssetTypeActions_SoundWave.cpp b/Engine/Source/Developer/AssetTools/Private/AssetTypeActions/AssetTypeActions_SoundWave.cpp
--- a/Engine/Source/Developer/AssetTools/Private/AssetTypeActions/AssetTypeActions_SoundWave.cpp
+++ b/Engine/Source/Developer/AssetTools/Private/AssetTypeActions/AssetTypeActions_SoundWave.cpp
@@ -5,6 +5,8 @@
 #include "SoundDefinitions.h"
 #include "Editor/UnrealEd/Public/Dialogs/DlgSoundWaveOptions.h"
 #include "ContentBrowserModule.h"
+#include "Sound/SoundWave.h"
+#include "Sound/SoundCue.h"
 
 #define LOCTEXT_NAMESPACE "AssetTypeActions"
 
diff --git a/Engine/Source/Developer/AssetTools/Private/AssetTypeActions/AssetTypeActions_StaticMesh.cpp b/Engine/Source/Developer/AssetTools/Private/AssetTypeActions/AssetTypeActions_StaticMesh.cpp
--- a/Engine/Source/Developer/AssetTools/Private/AssetTypeActions/AssetTypeActions_StaticMesh.cpp
+++ b/Engine/Source/Developer/AssetTools/Private/AssetTypeActions/AssetTypeActions_StaticMesh.cpp
@@ -13,6 +13,7 @@
 #include "FbxMeshUtils.h"
 #include "SNotificationList.h"
 #include "NotificationManager.h"
+#include "Engine/StaticMesh.h"
 
 #define LOCTEXT_NAMESPACE "AssetTypeActions"
 
diff --git a/Engine/Source/Developer/AssetTools/Private/AssetTypeActions/AssetTypeActions_StaticMesh.h b/Engine/Source/Developer/AssetTools/Private/AssetTypeActions/AssetTypeActions_StaticMesh.h
--- a/Engine/Source/Developer/AssetTools/Private/AssetTypeActions/AssetTypeActions_StaticMesh.h
+++ b/Engine/Source/Developer/AssetTools/Private/AssetTypeActions/AssetTypeActions_StaticMesh.h
@@ -1,6 +1,8 @@
 // Copyright 1998-2014 Epic Games, Inc. All Rights Reserved.
 
 #pragma once
+#include "Engine/StaticMesh.h"
+#include "AssetTypeActions_Base.h"
 
 class FAssetTypeActions_StaticMesh : public FAssetTypeActions_Base
 {
diff --git a/Engine/Source/Developer/AssetTools/Private/AssetTypeActions/AssetTypeActions_Struct.h b/Engine/Source/Developer/AssetTools/Private/AssetTypeActions/AssetTypeActions_Struct.h
--- a/Engine/Source/Developer/AssetTools/Private/AssetTypeActions/AssetTypeActions_Struct.h
+++ b/Engine/Source/Developer/AssetTools/Private/AssetTypeActions/AssetTypeActions_Struct.h
@@ -1,6 +1,8 @@
 // Copyright 1998-2013 Epic Games, Inc. All Rights Reserved.
 
 #pragma once
+#include "Engine/UserDefinedStruct.h"
+#include "AssetTypeActions_Base.h"
 
 class FAssetTypeActions_Struct : public FAssetTypeActions_Base
 {
diff --git a/Engine/Source/Developer/AssetTools/Private/AssetTypeActions/AssetTypeActions_SubsurfaceProfile.cpp b/Engine/Source/Developer/AssetTools/Private/AssetTypeActions/AssetTypeActions_SubsurfaceProfile.cpp
--- a/Engine/Source/Developer/AssetTools/Private/AssetTypeActions/AssetTypeActions_SubsurfaceProfile.cpp
+++ b/Engine/Source/Developer/AssetTools/Private/AssetTypeActions/AssetTypeActions_SubsurfaceProfile.cpp
@@ -1,8 +1,9 @@
 // Copyright 1998-2014 Epic Games, Inc. All Rights Reserved.
 
 #include "AssetToolsPrivatePCH.h"
+#include "Engine/SubsurfaceProfile.h"
 
 UClass* FAssetTypeActions_SubsurfaceProfile::GetSupportedClass() const
 {
 	return USubsurfaceProfile::StaticClass();
-}
\ No newline at end of file
+}
diff --git a/Engine/Source/Developer/AssetTools/Private/AssetTypeActions/AssetTypeActions_TextureCube.h b/Engine/Source/Developer/AssetTools/Private/AssetTypeActions/AssetTypeActions_TextureCube.h
--- a/Engine/Source/Developer/AssetTools/Private/AssetTypeActions/AssetTypeActions_TextureCube.h
+++ b/Engine/Source/Developer/AssetTools/Private/AssetTypeActions/AssetTypeActions_TextureCube.h
@@ -1,6 +1,8 @@
 // Copyright 1998-2014 Epic Games, Inc. All Rights Reserved.
 
 #pragma once
+#include "Engine/TextureCube.h"
+#include "AssetTypeActions_Texture.h"
 
 class FAssetTypeActions_TextureCube : public FAssetTypeActions_Texture
 {
diff --git a/Engine/Source/Developer/AssetTools/Private/AssetTypeActions/AssetTypeActions_TextureRenderTarget.cpp b/Engine/Source/Developer/AssetTools/Private/AssetTypeActions/AssetTypeActions_TextureRenderTarget.cpp
--- a/Engine/Source/Developer/AssetTools/Private/AssetTypeActions/AssetTypeActions_TextureRenderTarget.cpp
+++ b/Engine/Source/Developer/AssetTools/Private/AssetTypeActions/AssetTypeActions_TextureRenderTarget.cpp
@@ -2,6 +2,7 @@
 
 #include "AssetToolsPrivatePCH.h"
 #include "AssetRegistryModule.h"
+#include "Engine/TextureRenderTarget.h"
 
 #define LOCTEXT_NAMESPACE "AssetTypeActions"
 
diff --git a/Engine/Source/Developer/AssetTools/Private/AssetTypeActions/AssetTypeActions_TextureRenderTarget.h b/Engine/Source/Developer/AssetTools/Private/AssetTypeActions/AssetTypeActions_TextureRenderTarget.h
--- a/Engine/Source/Developer/AssetTools/Private/AssetTypeActions/AssetTypeActions_TextureRenderTarget.h
+++ b/Engine/Source/Developer/AssetTools/Private/AssetTypeActions/AssetTypeActions_TextureRenderTarget.h
@@ -1,6 +1,8 @@
 // Copyright 1998-2014 Epic Games, Inc. All Rights Reserved.
 
 #pragma once
+#include "Engine/TextureRenderTarget.h"
+#include "AssetTypeActions_Texture.h"
 
 class FAssetTypeActions_TextureRenderTarget : public FAssetTypeActions_Texture
 {
diff --git a/Engine/Source/Developer/AssetTools/Private/AssetTypeActions/AssetTypeActions_TextureRenderTarget2D.h b/Engine/Source/Developer/AssetTools/Private/AssetTypeActions/AssetTypeActions_TextureRenderTarget2D.h
--- a/Engine/Source/Developer/AssetTools/Private/AssetTypeActions/AssetTypeActions_TextureRenderTarget2D.h
+++ b/Engine/Source/Developer/AssetTools/Private/AssetTypeActions/AssetTypeActions_TextureRenderTarget2D.h
@@ -1,6 +1,8 @@
 // Copyright 1998-2014 Epic Games, Inc. All Rights Reserved.
 
 #pragma once
+#include "Engine/TextureRenderTarget2D.h"
+#include "AssetTypeActions_TextureRenderTarget.h"
 
 class FAssetTypeActions_TextureRenderTarget2D : public FAssetTypeActions_TextureRenderTarget
 {
diff --git a/Engine/Source/Developer/AssetTools/Private/AssetTypeActions/AssetTypeActions_TextureRenderTargetCube.h b/Engine/Source/Developer/AssetTools/Private/AssetTypeActions/AssetTypeActions_TextureRenderTargetCube.h
--- a/Engine/Source/Developer/AssetTools/Private/AssetTypeActions/AssetTypeActions_TextureRenderTargetCube.h
+++ b/Engine/Source/Developer/AssetTools/Private/AssetTypeActions/AssetTypeActions_TextureRenderTargetCube.h
@@ -1,6 +1,8 @@
 // Copyright 1998-2014 Epic Games, Inc. All Rights Reserved.
 
 #pragma once
+#include "Engine/TextureRenderTargetCube.h"
+#include "AssetTypeActions_TextureRenderTarget.h"
 
 class FAssetTypeActions_TextureRenderTargetCube : public FAssetTypeActions_TextureRenderTarget
 {
diff --git a/Engine/Source/Developer/AssetTools/Private/AssetTypeActions/AssetTypeActions_TouchInterface.cpp b/Engine/Source/Developer/AssetTools/Private/AssetTypeActions/AssetTypeActions_TouchInterface.cpp
--- a/Engine/Source/Developer/AssetTools/Private/AssetTypeActions/AssetTypeActions_TouchInterface.cpp
+++ b/Engine/Source/Developer/AssetTools/Private/AssetTypeActions/AssetTypeActions_TouchInterface.cpp
@@ -1,6 +1,7 @@
 // Copyright 1998-2014 Epic Games, Inc. All Rights Reserved.
 
 #include "AssetToolsPrivatePCH.h"
+#include "GameFramework/TouchInterface.h"
 
 UClass* FAssetTypeActions_TouchInterface::GetSupportedClass() const
 {
diff --git a/Engine/Source/Developer/MovieSceneTools/Private/AudioTrackEditor.cpp b/Engine/Source/Developer/MovieSceneTools/Private/AudioTrackEditor.cpp
--- a/Engine/Source/Developer/MovieSceneTools/Private/AudioTrackEditor.cpp
+++ b/Engine/Source/Developer/MovieSceneTools/Private/AudioTrackEditor.cpp
@@ -24,6 +24,9 @@
 #include "Runtime/Engine/Public/Slate/SceneViewport.h"
 #include "Runtime/MovieSceneCoreTypes/Classes/MovieSceneShotSection.h"
 #include "Runtime/Engine/Public/AudioDecompress.h"
+#include "Sound/SoundBase.h"
+#include "Sound/SoundWave.h"
+#include "Sound/SoundCue.h"
 
 
 namespace AnimatableAudioEditorConstants
diff --git a/Engine/Source/Developer/MovieSceneTools/Private/DirectorTrackEditor.cpp b/Engine/Source/Developer/MovieSceneTools/Private/DirectorTrackEditor.cpp
--- a/Engine/Source/Developer/MovieSceneTools/Private/DirectorTrackEditor.cpp
+++ b/Engine/Source/Developer/MovieSceneTools/Private/DirectorTrackEditor.cpp
@@ -20,7 +20,7 @@
 #include "MovieSceneTrackEditor.h"
 #include "DirectorTrackEditor.h"
 #include "CommonMovieSceneTools.h"
-
+#include "Camera/CameraActor.h"
 
 namespace AnimatableShotToolConstants
 {
diff --git a/Engine/Source/Developer/OutputLog/Private/SOutputLog.cpp b/Engine/Source/Developer/OutputLog/Private/SOutputLog.cpp
--- a/Engine/Source/Developer/OutputLog/Private/SOutputLog.cpp
+++ b/Engine/Source/Developer/OutputLog/Private/SOutputLog.cpp
@@ -4,7 +4,9 @@
 #include "SOutputLog.h"
 #include "SScrollBorder.h"
 #include "BaseTextLayoutMarshaller.h"
-
+#include "GameFramework/GameMode.h"
+#include "Engine/LocalPlayer.h"
+#include "GameFramework/GameState.h"
 /** Custom console editable text box whose only purpose is to prevent some keys from being typed */
 class SConsoleEditableTextBox : public SEditableTextBox
 {
diff --git a/Engine/Source/Developer/SuperSearch/Private/SSuperSearch.cpp b/Engine/Source/Developer/SuperSearch/Private/SSuperSearch.cpp
--- a/Engine/Source/Developer/SuperSearch/Private/SSuperSearch.cpp
+++ b/Engine/Source/Developer/SuperSearch/Private/SSuperSearch.cpp
@@ -7,6 +7,7 @@
 #include "IIntroTutorials.h"
 #include "EditorTutorial.h"
 #include "SSearchBox.h"
+#include "AssetData.h"
 
 static TSharedRef<FSearchEntry> OtherCategory(new FSearchEntry());
 static TSharedRef<FSearchEntry> AskQuestionEntry (new FSearchEntry());
diff --git a/Engine/Source/Developer/SuperSearch/Private/SSuperSearch.h b/Engine/Source/Developer/SuperSearch/Private/SSuperSearch.h
--- a/Engine/Source/Developer/SuperSearch/Private/SSuperSearch.h
+++ b/Engine/Source/Developer/SuperSearch/Private/SSuperSearch.h
@@ -1,6 +1,7 @@
 // Copyright 1998-2014 Epic Games, Inc. All Rights Reserved.
 
 #pragma once
+#include "AssetData.h"
 
 /**
  * Console input box with command-completion support
diff --git a/Engine/Source/Editor/Kismet/Private/FindInBlueprintManager.cpp b/Engine/Source/Editor/Kismet/Private/FindInBlueprintManager.cpp
--- a/Engine/Source/Editor/Kismet/Private/FindInBlueprintManager.cpp
+++ b/Engine/Source/Editor/Kismet/Private/FindInBlueprintManager.cpp
@@ -10,6 +10,7 @@
 #include "JsonUtilities.h"
 #include "SNotificationList.h"
 #include "NotificationManager.h"
+#include "Engine/SimpleConstructionScript.h"
 
 #define LOCTEXT_NAMESPACE "FindInBlueprintManager"
 
diff --git a/Engine/Source/Editor/PropertyEditor/Private/PropertyNode.cpp b/Engine/Source/Editor/PropertyEditor/Private/PropertyNode.cpp
--- a/Engine/Source/Editor/PropertyEditor/Private/PropertyNode.cpp
+++ b/Engine/Source/Editor/PropertyEditor/Private/PropertyNode.cpp
@@ -7,6 +7,7 @@
 #include "ScopedTransaction.h"
 #include "PropertyRestriction.h"
 #include "Editor/UnrealEd/Public/Kismet2/StructureEditorUtils.h"
+#include "Engine/UserDefinedStruct.h"
 
 FPropertySettings& FPropertySettings::Get()
 {
