Fix compile error for non-unity

diff --git a/Engine/Source/Developer/AssetTools/Private/AssetTypeActions/AssetTypeActions_Skeleton.cpp b/Engine/Source/Developer/AssetTools/Private/AssetTypeActions/AssetTypeActions_Skeleton.cpp
--- a/Engine/Source/Developer/AssetTools/Private/AssetTypeActions/AssetTypeActions_Skeleton.cpp
+++ b/Engine/Source/Developer/AssetTools/Private/AssetTypeActions/AssetTypeActions_Skeleton.cpp
@@ -12,7 +12,7 @@
 #include "SNotificationList.h"
 #include "SSkeletonWidget.h"
 #include "AnimationEditorUtils.h"
-
+#include "EditorAnimUtils.h"
 #include "Kismet2/KismetEditorUtilities.h"
 #include "Kismet2/BlueprintEditorUtils.h"
 #include "Animation/Rig.h"
diff --git a/Engine/Source/Editor/Persona/Private/SAnimCurveEd.cpp b/Engine/Source/Editor/Persona/Private/SAnimCurveEd.cpp
--- a/Engine/Source/Editor/Persona/Private/SAnimCurveEd.cpp
+++ b/Engine/Source/Editor/Persona/Private/SAnimCurveEd.cpp
@@ -5,6 +5,8 @@
 
 #include "SAnimCurveEd.h"
 
+#include "Editor/KismetWidgets/Public/SScrubWidget.h"
+
 #define LOCTEXT_NAMESPACE "AnimCurveEd"
 
 //////////////////////////////////////////////////////////////////////////
diff --git a/Engine/Source/Editor/Persona/Private/SAnimTrackCurvePanel.h b/Engine/Source/Editor/Persona/Private/SAnimTrackCurvePanel.h
--- a/Engine/Source/Editor/Persona/Private/SAnimTrackCurvePanel.h
+++ b/Engine/Source/Editor/Persona/Private/SAnimTrackCurvePanel.h
@@ -8,6 +8,7 @@
 #include "SNodePanel.h"
 #include "SAnimTrackPanel.h"
 #include "SCurveEditor.h"
+#include "SAnimNotifyPanel.h"
 #include "Animation/AnimSequenceBase.h"
 
 //////////////////////////////////////////////////////////////////////////
diff --git a/Engine/Source/Runtime/GameplayAbilities/Private/GameplayEffectTypes.cpp b/Engine/Source/Runtime/GameplayAbilities/Private/GameplayEffectTypes.cpp
--- a/Engine/Source/Runtime/GameplayAbilities/Private/GameplayEffectTypes.cpp
+++ b/Engine/Source/Runtime/GameplayAbilities/Private/GameplayEffectTypes.cpp
@@ -4,6 +4,7 @@
 #include "AbilitySystemInterface.h"
 #include "GameplayTagsModule.h"
 #include "GameplayEffectTypes.h"
+#include "AbilitySystemComponent.h"
 
 bool FGameplayEffectAttributeCaptureDefinition::operator==(const FGameplayEffectAttributeCaptureDefinition& Other) const
 {
