UE-5769 Crash when importing png to cube render target

diff --git a/Engine/Source/Developer/AssetTools/Private/AssetTypeActions/AssetTypeActions_TextureRenderTarget.h b/Engine/Source/Developer/AssetTools/Private/AssetTypeActions/AssetTypeActions_TextureRenderTarget.h
--- a/Engine/Source/Developer/AssetTools/Private/AssetTypeActions/AssetTypeActions_TextureRenderTarget.h
+++ b/Engine/Source/Developer/AssetTools/Private/AssetTypeActions/AssetTypeActions_TextureRenderTarget.h
@@ -13,6 +13,7 @@ public:
 	virtual UClass* GetSupportedClass() const override { return UTextureRenderTarget::StaticClass(); }
 	virtual bool HasActions ( const TArray<UObject*>& InObjects ) const override { return true; }
 	virtual void GetActions( const TArray<UObject*>& InObjects, FMenuBuilder& MenuBuilder ) override;
+	virtual bool IsImportedAsset() const override { return false; }
 
 private:
 	/** Handler for when CreateStatic is selected */
