Renaming an asset with no source control now prompts the user to connect to source control or make the file writable before performing the operaton.

diff --git a/Engine/Source/Developer/AssetTools/Private/AssetRenameManager.cpp b/Engine/Source/Developer/AssetTools/Private/AssetRenameManager.cpp
--- a/Engine/Source/Developer/AssetTools/Private/AssetRenameManager.cpp
+++ b/Engine/Source/Developer/AssetTools/Private/AssetRenameManager.cpp
@@ -432,31 +432,28 @@ bool FAssetRenameManager::CheckOutPackages(TArray<FAssetRenameDataWithReferencer
 	// Check out the packages
 	if (PackagesToCheckOut.Num() > 0)
 	{
-		if ( ISourceControlModule::Get().IsEnabled() )
+		TArray<UPackage*> PackagesCheckedOutOrMadeWritable;
+		TArray<UPackage*> PackagesNotNeedingCheckout;
+		bUserAcceptedCheckout = FEditorFileUtils::PromptToCheckoutPackages( false, PackagesToCheckOut, &PackagesCheckedOutOrMadeWritable, &PackagesNotNeedingCheckout );
+		if ( bUserAcceptedCheckout )
 		{
-			TArray<UPackage*> PackagesCheckedOutOrMadeWritable;
-			TArray<UPackage*> PackagesNotNeedingCheckout;
-			bUserAcceptedCheckout = FEditorFileUtils::PromptToCheckoutPackages( false, PackagesToCheckOut, &PackagesCheckedOutOrMadeWritable, &PackagesNotNeedingCheckout );
-			if ( bUserAcceptedCheckout )
-			{
-				// Make a list of any packages in the list which weren't checked out for some reason
-				TArray<UPackage*> PackagesThatCouldNotBeCheckedOut = PackagesToCheckOut;
+			// Make a list of any packages in the list which weren't checked out for some reason
+			TArray<UPackage*> PackagesThatCouldNotBeCheckedOut = PackagesToCheckOut;
 
-				for ( auto PackageIt = PackagesCheckedOutOrMadeWritable.CreateConstIterator(); PackageIt; ++PackageIt )
-				{
-					PackagesThatCouldNotBeCheckedOut.Remove(*PackageIt);
-				}
+			for ( auto PackageIt = PackagesCheckedOutOrMadeWritable.CreateConstIterator(); PackageIt; ++PackageIt )
+			{
+				PackagesThatCouldNotBeCheckedOut.Remove(*PackageIt);
+			}
 
-				for ( auto PackageIt = PackagesNotNeedingCheckout.CreateConstIterator(); PackageIt; ++PackageIt )
-				{
-					PackagesThatCouldNotBeCheckedOut.Remove(*PackageIt);
-				}
+			for ( auto PackageIt = PackagesNotNeedingCheckout.CreateConstIterator(); PackageIt; ++PackageIt )
+			{
+				PackagesThatCouldNotBeCheckedOut.Remove(*PackageIt);
+			}
 
-				// If there's anything which couldn't be checked out, abort the operation.
-				if (PackagesThatCouldNotBeCheckedOut.Num() > 0)
-				{
-					bUserAcceptedCheckout = false;
-				}
+			// If there's anything which couldn't be checked out, abort the operation.
+			if (PackagesThatCouldNotBeCheckedOut.Num() > 0)
+			{
+				bUserAcceptedCheckout = false;
 			}
 		}
 	}
