[INTEGRATE] Change 2417072 by [REDACTED_PERSON][REDACTED_HANDLE] on 2015/01/23 13:55:54

	Toned down some of the bright blue asset colors to be more pastel

	UE-7184 - Show game c++ classes and engine c++ classes in the content browser

	These kinds of colors work better with the new asset thumbnails.

diff --git a/Engine/Plugins/Media/MediaPlayerEditor/Source/MediaPlayerEditor/Private/AssetTools/MediaSoundWaveActions.cpp b/Engine/Plugins/Media/MediaPlayerEditor/Source/MediaPlayerEditor/Private/AssetTools/MediaSoundWaveActions.cpp
--- a/Engine/Plugins/Media/MediaPlayerEditor/Source/MediaPlayerEditor/Private/AssetTools/MediaSoundWaveActions.cpp
+++ b/Engine/Plugins/Media/MediaPlayerEditor/Source/MediaPlayerEditor/Private/AssetTools/MediaSoundWaveActions.cpp
@@ -35,7 +35,7 @@ UClass* FMediaSoundWaveActions::GetSupportedClass() const
 
 FColor FMediaSoundWaveActions::GetTypeColor() const
 {
-	return FColor::Blue;
+	return FColor(77, 93, 239);
 }
 
 
diff --git a/Engine/Source/Developer/AssetTools/Private/AssetTypeActions/AssetTypeActions_DialogueVoice.h b/Engine/Source/Developer/AssetTools/Private/AssetTypeActions/AssetTypeActions_DialogueVoice.h
--- a/Engine/Source/Developer/AssetTools/Private/AssetTypeActions/AssetTypeActions_DialogueVoice.h
+++ b/Engine/Source/Developer/AssetTools/Private/AssetTypeActions/AssetTypeActions_DialogueVoice.h
@@ -10,7 +10,7 @@ class FAssetTypeActions_DialogueVoice : public FAssetTypeActions_Base
 public:
 	// IAssetTypeActions Implementation
 	virtual FText GetName() const override { return NSLOCTEXT("AssetTypeActions", "AssetTypeActions_DialogueVoice", "Dialogue Voice"); }
-	virtual FColor GetTypeColor() const override { return FColor(0,0,255); }
+	virtual FColor GetTypeColor() const override { return FColor(97, 85, 212); }
 	virtual UClass* GetSupportedClass() const override { return UDialogueVoice::StaticClass(); }
 	virtual uint32 GetCategories() override { return EAssetTypeCategories::Sounds; }
 	virtual bool CanFilter() override { return true; }
diff --git a/Engine/Source/Developer/AssetTools/Private/AssetTypeActions/AssetTypeActions_DialogueWave.h b/Engine/Source/Developer/AssetTools/Private/AssetTypeActions/AssetTypeActions_DialogueWave.h
--- a/Engine/Source/Developer/AssetTools/Private/AssetTypeActions/AssetTypeActions_DialogueWave.h
+++ b/Engine/Source/Developer/AssetTools/Private/AssetTypeActions/AssetTypeActions_DialogueWave.h
@@ -10,7 +10,7 @@ class FAssetTypeActions_DialogueWave : public FAssetTypeActions_Base
 public:
 	// IAssetTypeActions Implementation
 	virtual FText GetName() const override { return NSLOCTEXT("AssetTypeActions", "AssetTypeActions_DialogueWave", "Dialogue Wave"); }
-	virtual FColor GetTypeColor() const override { return FColor(0,0,255); }
+	virtual FColor GetTypeColor() const override { return FColor(97, 85, 212); }
 	virtual UClass* GetSupportedClass() const override { return UDialogueWave::StaticClass(); }
 	virtual uint32 GetCategories() override { return EAssetTypeCategories::Sounds; }
 	virtual bool CanFilter() override { return true; }
diff --git a/Engine/Source/Developer/AssetTools/Private/AssetTypeActions/AssetTypeActions_SoundAttenuation.h b/Engine/Source/Developer/AssetTools/Private/AssetTypeActions/AssetTypeActions_SoundAttenuation.h
--- a/Engine/Source/Developer/AssetTools/Private/AssetTypeActions/AssetTypeActions_SoundAttenuation.h
+++ b/Engine/Source/Developer/AssetTools/Private/AssetTypeActions/AssetTypeActions_SoundAttenuation.h
@@ -7,7 +7,7 @@ class FAssetTypeActions_SoundAttenuation : public FAssetTypeActions_Base
 public:
 	// IAssetTypeActions Implementation
 	virtual FText GetName() const override { return NSLOCTEXT("AssetTypeActions", "AssetTypeActions_SoundAttenuation", "Sound Attenuation"); }
-	virtual FColor GetTypeColor() const override { return FColor(0, 0, 175); }
+	virtual FColor GetTypeColor() const override { return FColor(77, 120, 239); }
 	virtual UClass* GetSupportedClass() const override;
 	virtual uint32 GetCategories() override { return EAssetTypeCategories::Sounds; }
 };
\ No newline at end of file
diff --git a/Engine/Source/Developer/AssetTools/Private/AssetTypeActions/AssetTypeActions_SoundBase.h b/Engine/Source/Developer/AssetTools/Private/AssetTypeActions/AssetTypeActions_SoundBase.h
--- a/Engine/Source/Developer/AssetTools/Private/AssetTypeActions/AssetTypeActions_SoundBase.h
+++ b/Engine/Source/Developer/AssetTools/Private/AssetTypeActions/AssetTypeActions_SoundBase.h
@@ -7,7 +7,7 @@ class FAssetTypeActions_SoundBase : public FAssetTypeActions_Base
 public:
 	// IAssetTypeActions Implementation
 	virtual FText GetName() const override { return NSLOCTEXT("AssetTypeActions", "AssetTypeActions_SoundBase", "Sound Base"); }
-	virtual FColor GetTypeColor() const override { return FColor(0,0,255); }
+	virtual FColor GetTypeColor() const override { return FColor(97, 85, 212); }
 	virtual UClass* GetSupportedClass() const override;
 	virtual bool HasActions ( const TArray<UObject*>& InObjects ) const override { return true; }
 	virtual void GetActions( const TArray<UObject*>& InObjects, FMenuBuilder& MenuBuilder ) override;
diff --git a/Engine/Source/Developer/AssetTools/Private/AssetTypeActions/AssetTypeActions_SoundWave.h b/Engine/Source/Developer/AssetTools/Private/AssetTypeActions/AssetTypeActions_SoundWave.h
--- a/Engine/Source/Developer/AssetTools/Private/AssetTypeActions/AssetTypeActions_SoundWave.h
+++ b/Engine/Source/Developer/AssetTools/Private/AssetTypeActions/AssetTypeActions_SoundWave.h
@@ -7,7 +7,7 @@ class FAssetTypeActions_SoundWave : public FAssetTypeActions_SoundBase
 public:
 	// IAssetTypeActions Implementation
 	virtual FText GetName() const override { return NSLOCTEXT("AssetTypeActions", "AssetTypeActions_SoundWave", "Sound Wave"); }
-	virtual FColor GetTypeColor() const override { return FColor(0, 0, 255); }
+	virtual FColor GetTypeColor() const override { return FColor(97, 85, 212); }
 	virtual UClass* GetSupportedClass() const override;
 	virtual void GetActions( const TArray<UObject*>& InObjects, FMenuBuilder& MenuBuilder ) override;
 	virtual void OpenAssetEditor( const TArray<UObject*>& InObjects, TSharedPtr<class IToolkitHost> EditWithinLevelEditor = TSharedPtr<IToolkitHost>() ) override;
diff --git a/Engine/Source/Developer/AssetTools/Private/AssetTypeActions/AssetTypeActions_Struct.h b/Engine/Source/Developer/AssetTools/Private/AssetTypeActions/AssetTypeActions_Struct.h
--- a/Engine/Source/Developer/AssetTools/Private/AssetTypeActions/AssetTypeActions_Struct.h
+++ b/Engine/Source/Developer/AssetTools/Private/AssetTypeActions/AssetTypeActions_Struct.h
@@ -9,7 +9,7 @@ class FAssetTypeActions_Struct : public FAssetTypeActions_Base
 public:
 	// IAssetTypeActions Implementation
 	virtual FText GetName() const override { return NSLOCTEXT("AssetTypeActions", "AssetTypeActions_Struct", "Structure"); }
-	virtual FColor GetTypeColor() const override { return FColor(0, 0, 255); }
+	virtual FColor GetTypeColor() const override { return FColor(103, 206, 218); }
 	virtual UClass* GetSupportedClass() const override { return UUserDefinedStruct::StaticClass(); }
 	virtual uint32 GetCategories() override { return EAssetTypeCategories::Blueprint; }
 
