[INTEGRATE] Change 2418630 by [REDACTED_PERSON][REDACTED_HANDLE] on 2015/01/26 10:50:03

	Added the concept of "ClassTypeActions" for class asset types (C++ classes and Blueprints)

	UE-7970 - Add small 'play' button to EditorTutorial asset thumbnail

	Class asset types (C++ classes and Blueprints) are now able to provide actions that should be performed based upon the type of class the asset contains, rather than the asset class itself (which is always "Class" or "Blueprint").

	This will allow the editor tutorials to provide class type actions for Blueprints that derive from "EditorTutorial" without requiring this information to be put into the asset type actions for Blueprint assets.

diff --git a/Engine/Source/Developer/AssetTools/Private/AssetTools.cpp b/Engine/Source/Developer/AssetTools/Private/AssetTools.cpp
--- a/Engine/Source/Developer/AssetTools/Private/AssetTools.cpp
+++ b/Engine/Source/Developer/AssetTools/Private/AssetTools.cpp
@@ -149,6 +149,45 @@ TWeakPtr<IAssetTypeActions> FAssetTools::GetAssetTypeActionsForClass( UClass* Cl
 	return MostDerivedAssetTypeActions;
 }
 
+void FAssetTools::RegisterClassTypeActions(const TSharedRef<IClassTypeActions>& NewActions)
+{
+	ClassTypeActionsList.Add(NewActions);
+}
+
+void FAssetTools::UnregisterClassTypeActions(const TSharedRef<IClassTypeActions>& ActionsToRemove)
+{
+	ClassTypeActionsList.Remove(ActionsToRemove);
+}
+
+void FAssetTools::GetClassTypeActionsList( TArray<TWeakPtr<IClassTypeActions>>& OutClassTypeActionsList ) const
+{
+	for (auto ActionsIt = ClassTypeActionsList.CreateConstIterator(); ActionsIt; ++ActionsIt)
+	{
+		OutClassTypeActionsList.Add(*ActionsIt);
+	}
+}
+
+TWeakPtr<IClassTypeActions> FAssetTools::GetClassTypeActionsForClass( UClass* Class ) const
+{
+	TSharedPtr<IClassTypeActions> MostDerivedClassTypeActions;
+
+	for (int32 TypeActionsIdx = 0; TypeActionsIdx < ClassTypeActionsList.Num(); ++TypeActionsIdx)
+	{
+		TSharedRef<IClassTypeActions> TypeActions = ClassTypeActionsList[TypeActionsIdx];
+		UClass* SupportedClass = TypeActions->GetSupportedClass();
+
+		if ( Class->IsChildOf(SupportedClass) )
+		{
+			if ( !MostDerivedClassTypeActions.IsValid() || SupportedClass->IsChildOf( MostDerivedClassTypeActions->GetSupportedClass() ) )
+			{
+				MostDerivedClassTypeActions = TypeActions;
+			}
+		}
+	}
+
+	return MostDerivedClassTypeActions;
+}
+
 bool FAssetTools::GetAssetActions( const TArray<UObject*>& InObjects, FMenuBuilder& MenuBuilder, bool bIncludeHeading )
 {
 	bool bAddedActions = false;
diff --git a/Engine/Source/Developer/AssetTools/Private/AssetTools.h b/Engine/Source/Developer/AssetTools/Private/AssetTools.h
--- a/Engine/Source/Developer/AssetTools/Private/AssetTools.h
+++ b/Engine/Source/Developer/AssetTools/Private/AssetTools.h
@@ -13,6 +13,10 @@ public:
 	virtual void UnregisterAssetTypeActions(const TSharedRef<IAssetTypeActions>& ActionsToRemove) override;
 	virtual void GetAssetTypeActionsList( TArray<TWeakPtr<IAssetTypeActions>>& OutAssetTypeActionsList ) const override;
 	virtual TWeakPtr<IAssetTypeActions> GetAssetTypeActionsForClass( UClass* Class ) const override;
+	virtual void RegisterClassTypeActions(const TSharedRef<IClassTypeActions>& NewActions) override;
+	virtual void UnregisterClassTypeActions(const TSharedRef<IClassTypeActions>& ActionsToRemove) override;
+	virtual void GetClassTypeActionsList( TArray<TWeakPtr<IClassTypeActions>>& OutClassTypeActionsList ) const override;
+	virtual TWeakPtr<IClassTypeActions> GetClassTypeActionsForClass( UClass* Class ) const override;
 	virtual bool GetAssetActions( const TArray<UObject*>& InObjects, FMenuBuilder& MenuBuilder, bool bIncludeHeading = true ) override;
 	virtual UObject* CreateAsset(const FString& AssetName, const FString& PackagePath, UClass* AssetClass, UFactory* Factory, FName CallingContext = NAME_None) override;
 	virtual UObject* CreateAsset(UClass* AssetClass, UFactory* Factory, FName CallingContext = NAME_None) override;
@@ -71,4 +75,7 @@ private:
 
 	/** The list of all registered AssetTypeActions */
 	TArray<TSharedRef<IAssetTypeActions>> AssetTypeActionsList;
+
+	/** The list of all registered ClassTypeActions */
+	TArray<TSharedRef<IClassTypeActions>> ClassTypeActionsList;
 };
diff --git a/Engine/Source/Developer/AssetTools/Private/AssetToolsPrivatePCH.h b/Engine/Source/Developer/AssetTools/Private/AssetToolsPrivatePCH.h
--- a/Engine/Source/Developer/AssetTools/Private/AssetToolsPrivatePCH.h
+++ b/Engine/Source/Developer/AssetTools/Private/AssetToolsPrivatePCH.h
@@ -9,11 +9,15 @@
 DECLARE_LOG_CATEGORY_EXTERN(LogAssetTools, Log, All);
 
 #include "IAssetTypeActions.h"
+#include "IClassTypeActions.h"
 #include "IAssetTools.h"
 
 #include "AssetTypeActions_Base.h"
 #include "AssetTypeActions_CSVAssetBase.h"
 
+#include "ClassTypeActions_Base.h"
+#include "AssetTypeActions/AssetTypeActions_ClassTypeBase.h"
+
 #include "AssetTypeActions/AssetTypeActions_Blueprint.h"
 #include "AssetTypeActions/AssetTypeActions_Curve.h"
 #include "AssetTypeActions/AssetTypeActions_MaterialInterface.h"
diff --git a/Engine/Source/Developer/AssetTools/Private/AssetTypeActions/AssetTypeActions_Blueprint.cpp b/Engine/Source/Developer/AssetTools/Private/AssetTypeActions/AssetTypeActions_Blueprint.cpp
--- a/Engine/Source/Developer/AssetTools/Private/AssetTypeActions/AssetTypeActions_Blueprint.cpp
+++ b/Engine/Source/Developer/AssetTools/Private/AssetTypeActions/AssetTypeActions_Blueprint.cpp
@@ -294,4 +294,33 @@ FText FAssetTypeActions_Blueprint::GetAssetDescription(const FAssetData& AssetDa
 	return FText::GetEmpty();
 }
 
+TWeakPtr<IClassTypeActions> FAssetTypeActions_Blueprint::GetClassTypeActions(const FAssetData& AssetData) const
+{
+	static const FName NativeParentClassTag = "NativeParentClass";
+	static const FName ParentClassTag = "ParentClass";
+
+	// Blueprints get the class type actions for their parent native class - this avoids us having to load the blueprint
+	UClass* ParentClass = nullptr;
+	const FString* ParentClassNamePtr = AssetData.TagsAndValues.Find(NativeParentClassTag);
+	if(!ParentClassNamePtr)
+	{
+		ParentClassNamePtr = AssetData.TagsAndValues.Find(ParentClassTag);
+	}
+	if(ParentClassNamePtr && !ParentClassNamePtr->IsEmpty())
+	{
+		UObject* Outer = nullptr;
+		FString ParentClassName = *ParentClassNamePtr;
+		ResolveName(Outer, ParentClassName, false, false);
+		ParentClass = FindObject<UClass>(ANY_PACKAGE, *ParentClassName);
+	}
+
+	if(ParentClass)
+	{
+		FAssetToolsModule& AssetToolsModule = FAssetToolsModule::GetModule();
+		return AssetToolsModule.Get().GetClassTypeActionsForClass(ParentClass);
+	}
+
+	return nullptr;
+}
+
 #undef LOCTEXT_NAMESPACE
diff --git a/Engine/Source/Developer/AssetTools/Private/AssetTypeActions/AssetTypeActions_Blueprint.h b/Engine/Source/Developer/AssetTools/Private/AssetTypeActions/AssetTypeActions_Blueprint.h
--- a/Engine/Source/Developer/AssetTools/Private/AssetTypeActions/AssetTypeActions_Blueprint.h
+++ b/Engine/Source/Developer/AssetTools/Private/AssetTypeActions/AssetTypeActions_Blueprint.h
@@ -2,9 +2,9 @@
 
 #pragma once
 
-#include "AssetTypeActions_Base.h"
+#include "AssetTypeActions_ClassTypeBase.h"
 
-class ASSETTOOLS_API FAssetTypeActions_Blueprint : public FAssetTypeActions_Base
+class ASSETTOOLS_API FAssetTypeActions_Blueprint : public FAssetTypeActions_ClassTypeBase
 {
 public:
 	// IAssetTypeActions Implementation
@@ -22,6 +22,9 @@ public:
 	virtual class UThumbnailInfo* GetThumbnailInfo(UObject* Asset) const override;
 	virtual FText GetAssetDescription(const FAssetData& AssetData) const override;
 
+	// FAssetTypeActions_ClassTypeBase Implementation
+	virtual TWeakPtr<IClassTypeActions> GetClassTypeActions(const FAssetData& AssetData) const override;
+
 protected:
 	/** Whether or not this asset can create derived blueprints */
 	virtual bool CanCreateNewDerivedBlueprint() const;
diff --git a/Engine/Source/Developer/AssetTools/Private/AssetTypeActions/AssetTypeActions_Class.cpp b/Engine/Source/Developer/AssetTools/Private/AssetTypeActions/AssetTypeActions_Class.cpp
--- a/Engine/Source/Developer/AssetTools/Private/AssetTypeActions/AssetTypeActions_Class.cpp
+++ b/Engine/Source/Developer/AssetTools/Private/AssetTypeActions/AssetTypeActions_Class.cpp
@@ -95,4 +95,16 @@ void FAssetTypeActions_Class::OpenAssetEditor( const TArray<UObject*>& InObjects
 	}
 }
 
+TWeakPtr<IClassTypeActions> FAssetTypeActions_Class::GetClassTypeActions(const FAssetData& AssetData) const
+{
+	UClass* Class = FindObject<UClass>(ANY_PACKAGE, *AssetData.AssetName.ToString());
+	if(Class)
+	{
+		FAssetToolsModule& AssetToolsModule = FAssetToolsModule::GetModule();
+		return AssetToolsModule.Get().GetClassTypeActionsForClass(Class);
+	}
+
+	return nullptr;
+}
+
 #undef LOCTEXT_NAMESPACE
diff --git a/Engine/Source/Developer/AssetTools/Private/AssetTypeActions/AssetTypeActions_Class.h b/Engine/Source/Developer/AssetTools/Private/AssetTypeActions/AssetTypeActions_Class.h
--- a/Engine/Source/Developer/AssetTools/Private/AssetTypeActions/AssetTypeActions_Class.h
+++ b/Engine/Source/Developer/AssetTools/Private/AssetTypeActions/AssetTypeActions_Class.h
@@ -2,9 +2,9 @@
 
 #pragma once
 
-#include "AssetTypeActions_Base.h"
+#include "AssetTypeActions_ClassTypeBase.h"
 
-class FAssetTypeActions_Class : public FAssetTypeActions_Base
+class FAssetTypeActions_Class : public FAssetTypeActions_ClassTypeBase
 {
 public:
 	// IAssetTypeActions Implementation
@@ -19,4 +19,7 @@ public:
 	virtual class UThumbnailInfo* GetThumbnailInfo(UObject* Asset) const override;
 
 	virtual void OpenAssetEditor( const TArray<UObject*>& InObjects, TSharedPtr<class IToolkitHost> EditWithinLevelEditor = TSharedPtr<IToolkitHost>() ) override;
+
+	// FAssetTypeActions_ClassTypeBase Implementation
+	virtual TWeakPtr<IClassTypeActions> GetClassTypeActions(const FAssetData& AssetData) const override;
 };
diff --git a/Engine/Source/Developer/AssetTools/Private/AssetTypeActions/AssetTypeActions_ClassTypeBase.cpp b/Engine/Source/Developer/AssetTools/Private/AssetTypeActions/AssetTypeActions_ClassTypeBase.cpp
new file mode 100644
--- /dev/null
+++ b/Engine/Source/Developer/AssetTools/Private/AssetTypeActions/AssetTypeActions_ClassTypeBase.cpp
@@ -0,0 +1,18 @@
+// Copyright 1998-2015 Epic Games, Inc. All Rights Reserved.
+
+#include "AssetToolsPrivatePCH.h"
+
+#define LOCTEXT_NAMESPACE "AssetTypeActions"
+
+TSharedPtr<SWidget> FAssetTypeActions_ClassTypeBase::GetThumbnailOverlay(const FAssetData& AssetData) const
+{
+	TSharedPtr<IClassTypeActions> ClassTypeActions = GetClassTypeActions(AssetData).Pin();
+	if(ClassTypeActions.IsValid())
+	{
+		return ClassTypeActions->GetThumbnailOverlay(AssetData);
+	}
+
+	return nullptr;
+}
+
+#undef LOCTEXT_NAMESPACE
diff --git a/Engine/Source/Developer/AssetTools/Private/AssetTypeActions/AssetTypeActions_ClassTypeBase.h b/Engine/Source/Developer/AssetTools/Private/AssetTypeActions/AssetTypeActions_ClassTypeBase.h
new file mode 100644
--- /dev/null
+++ b/Engine/Source/Developer/AssetTools/Private/AssetTypeActions/AssetTypeActions_ClassTypeBase.h
@@ -0,0 +1,18 @@
+// Copyright 1998-2015 Epic Games, Inc. All Rights Reserved.
+
+#pragma once
+
+#include "AssetTypeActions_Base.h"
+
+class IClassTypeActions;
+
+/** Base class for "class type" assets (C++ classes and Blueprints */
+class ASSETTOOLS_API FAssetTypeActions_ClassTypeBase : public FAssetTypeActions_Base
+{
+public:
+	/** Get the class type actions for this asset */
+	virtual TWeakPtr<IClassTypeActions> GetClassTypeActions(const FAssetData& AssetData) const = 0;
+
+	// IAssetTypeActions Implementation
+	virtual TSharedPtr<class SWidget> GetThumbnailOverlay(const FAssetData& AssetData) const override;
+};
diff --git a/Engine/Source/Developer/AssetTools/Public/AssetToolsModule.h b/Engine/Source/Developer/AssetTools/Public/AssetToolsModule.h
--- a/Engine/Source/Developer/AssetTools/Public/AssetToolsModule.h
+++ b/Engine/Source/Developer/AssetTools/Public/AssetToolsModule.h
@@ -4,6 +4,7 @@
 
 #include "ModuleManager.h"
 #include "IAssetTypeActions.h"
+#include "IClassTypeActions.h"
 #include "IAssetTools.h"
 
 class FAssetToolsModule : public IModuleInterface
@@ -22,7 +23,8 @@ public:
 	 */
 	static inline FAssetToolsModule& GetModule()
 	{
-		return FModuleManager::LoadModuleChecked< FAssetToolsModule >("AssetTools");
+		static const FName ModuleName = "AssetTools";
+		return FModuleManager::LoadModuleChecked< FAssetToolsModule >(ModuleName);
 	}
 
 private:
diff --git a/Engine/Source/Developer/AssetTools/Public/ClassTypeActions_Base.h b/Engine/Source/Developer/AssetTools/Public/ClassTypeActions_Base.h
new file mode 100644
--- /dev/null
+++ b/Engine/Source/Developer/AssetTools/Public/ClassTypeActions_Base.h
@@ -0,0 +1,17 @@
+// Copyright 1998-2015 Epic Games, Inc. All Rights Reserved.
+
+#pragma once
+
+#include "AssetToolsModule.h"
+
+/** A base class for all ClassTypeActions. Provides helper functions useful for many types. Deriving from this class is optional. */
+class FClassTypeActions_Base : public IClassTypeActions
+{
+public:
+	// Begin IClassTypeActions implementation
+	virtual TSharedPtr<class SWidget> GetThumbnailOverlay(const FAssetData& AssetData) const override
+	{
+		return nullptr;
+	}
+	// End IAssetTypeActions implementation
+};
diff --git a/Engine/Source/Developer/AssetTools/Public/IAssetTools.h b/Engine/Source/Developer/AssetTools/Public/IAssetTools.h
--- a/Engine/Source/Developer/AssetTools/Public/IAssetTools.h
+++ b/Engine/Source/Developer/AssetTools/Public/IAssetTools.h
@@ -3,6 +3,7 @@
 #pragma once
 
 class IAssetTypeActions;
+class IClassTypeActions;
 class FAssetData;
 class UFactory;
 
@@ -43,6 +44,18 @@ public:
 	/** Gets the appropriate AssetTypeActions for the supplied class */
 	virtual TWeakPtr<IAssetTypeActions> GetAssetTypeActionsForClass( UClass* Class ) const = 0;
 
+	/** Registers a class type actions object so it can provide information about and actions for class asset types. */
+	virtual void RegisterClassTypeActions(const TSharedRef<IClassTypeActions>& NewActions) = 0;
+
+	/** Unregisters a class type actions object. It will no longer provide information about or actions for class asset types. */
+	virtual void UnregisterClassTypeActions(const TSharedRef<IClassTypeActions>& ActionsToRemove) = 0;
+
+	/** Generates a list of currently registered ClassTypeActions */
+	virtual void GetClassTypeActionsList( TArray<TWeakPtr<IClassTypeActions>>& OutClassTypeActionsList ) const = 0;
+
+	/** Gets the appropriate ClassTypeActions for the supplied class */
+	virtual TWeakPtr<IClassTypeActions> GetClassTypeActionsForClass( UClass* Class ) const = 0;
+
 	/**
 	 * Fills out a menubuilder with a list of commands that can be applied to the specified objects.
 	 *
diff --git a/Engine/Source/Developer/AssetTools/Public/IClassTypeActions.h b/Engine/Source/Developer/AssetTools/Public/IClassTypeActions.h
new file mode 100644
--- /dev/null
+++ b/Engine/Source/Developer/AssetTools/Public/IClassTypeActions.h
@@ -0,0 +1,17 @@
+// Copyright 1998-2015 Epic Games, Inc. All Rights Reserved.
+
+#pragma once
+
+/** ClassTypeActions provide actions and other information about asset types that host classes */
+class IClassTypeActions : public TSharedFromThis<IClassTypeActions>
+{
+public:
+	/** Virtual destructor */
+	virtual ~IClassTypeActions(){}
+
+	/** Checks to see if the specified object is handled by this type. */
+	virtual UClass* GetSupportedClass() const = 0;
+
+	/** Optionally returns a custom widget to overlay on top of this assets' thumbnail */
+	virtual TSharedPtr<class SWidget> GetThumbnailOverlay(const class FAssetData& AssetData) const = 0;
+};
