Fixed bad logic in FAssetTypeActions_Font::ExecuteReimport

UE-8477 - Can't reimport offline cached font asset

diff --git a/Engine/Source/Developer/AssetTools/Private/AssetTypeActions/AssetTypeActions_Font.cpp b/Engine/Source/Developer/AssetTools/Private/AssetTypeActions/AssetTypeActions_Font.cpp
--- a/Engine/Source/Developer/AssetTools/Private/AssetTypeActions/AssetTypeActions_Font.cpp
+++ b/Engine/Source/Developer/AssetTools/Private/AssetTypeActions/AssetTypeActions_Font.cpp
@@ -61,7 +61,7 @@ void FAssetTypeActions_Font::ExecuteReimport(const TArray<TWeakObjectPtr<UFont>>
 		if (Object)
 		{
 			// Skip fonts that aren't using an offline cache, as they can't be reimported
-			if (Object->FontCacheType != EFontCacheType::Offline)
+			if (Object->FontCacheType == EFontCacheType::Offline)
 			{
 				// Fonts fail to reimport if they ask for a new file if missing
 				FReimportManager::Instance()->Reimport(Object, /*bAskForNewFileIfMissing=*/false);
