[INTEGRATE] Change 2435229 by bruce.nesbit@bnesbit_Releases on 2015/02/06 08:09:07

	Fixed issue with editor layout when inserting feature pack during editor startup
	(Fixes Jira 9223)

diff --git a/Engine/Source/Developer/AssetTools/Private/AssetTools.cpp b/Engine/Source/Developer/AssetTools/Private/AssetTools.cpp
--- a/Engine/Source/Developer/AssetTools/Private/AssetTools.cpp
+++ b/Engine/Source/Developer/AssetTools/Private/AssetTools.cpp
@@ -511,7 +511,7 @@ void FAssetTools::ExpandDirectories(const TArray<FString>& Files, const FString&
 	}
 }
 
-TArray<UObject*> FAssetTools::ImportAssets(const TArray<FString>& Files, const FString& RootDestinationPath, UFactory* SpecifiedFactory) const
+TArray<UObject*> FAssetTools::ImportAssets(const TArray<FString>& Files, const FString& RootDestinationPath, UFactory* SpecifiedFactory, bool bSyncToBrowser ) const
 {
 	TArray<UObject*> ReturnObjects;
 	TMap< FString, TArray<UFactory*> > ExtensionToFactoriesMap;
@@ -875,7 +875,7 @@ TArray<UObject*> FAssetTools::ImportAssets(const TArray<FString>& Files, const F
 	}
 
 	// Sync content browser to the newly created assets
-	if ( ReturnObjects.Num() )
+	if ( ReturnObjects.Num() && (bSyncToBrowser != false ) )
 	{
 		FAssetTools::Get().SyncBrowserToAssets(ReturnObjects);
 	}
diff --git a/Engine/Source/Developer/AssetTools/Private/AssetTools.h b/Engine/Source/Developer/AssetTools/Private/AssetTools.h
--- a/Engine/Source/Developer/AssetTools/Private/AssetTools.h
+++ b/Engine/Source/Developer/AssetTools/Private/AssetTools.h
@@ -23,7 +23,7 @@ public:
 	virtual UObject* DuplicateAsset(const FString& AssetName, const FString& PackagePath, UObject* OriginalObject) override;
 	virtual void RenameAssets(const TArray<FAssetRenameData>& AssetsAndNames) const override;
 	virtual TArray<UObject*> ImportAssets(const FString& DestinationPath) override;
-	virtual TArray<UObject*> ImportAssets(const TArray<FString>& Files, const FString& DestinationPath, UFactory* ChosenFactory) const override;
+	virtual TArray<UObject*> ImportAssets(const TArray<FString>& Files, const FString& DestinationPath, UFactory* ChosenFactory, bool bSyncToBrowser = true) const override;
 	virtual void CreateUniqueAssetName(const FString& InBasePackageName, const FString& InSuffix, FString& OutPackageName, FString& OutAssetName) const override;
 	virtual bool AssetUsesGenericThumbnail( const FAssetData& AssetData ) const override;
 	virtual void DiffAgainstDepot(UObject* InObject, const FString& InPackagePath, const FString& InPackageName) const override;
diff --git a/Engine/Source/Developer/AssetTools/Public/IAssetTools.h b/Engine/Source/Developer/AssetTools/Public/IAssetTools.h
--- a/Engine/Source/Developer/AssetTools/Public/IAssetTools.h
+++ b/Engine/Source/Developer/AssetTools/Public/IAssetTools.h
@@ -90,8 +90,15 @@ public:
 	/** Opens a file open dialog to choose files to import to the destination path. */
 	virtual TArray<UObject*> ImportAssets(const FString& DestinationPath) = 0;
 
-	/** Imports the specified files to the destination path. */
-	virtual TArray<UObject*> ImportAssets(const TArray<FString>& Files, const FString& DestinationPath, UFactory* ChosenFactory = NULL) const = 0;
+	/** 
+	 * Imports the specified files to the destination path. 
+	 *
+	 * @param Files				Files to import
+	 * @param DestinationPath	destination path for imported files
+	 * @param ChosenFactory		Specific factory to use for object creation
+	 * @param bSyncToBrowser	If true sync content browser to first imported asset after import
+	 */
+	virtual TArray<UObject*> ImportAssets(const TArray<FString>& Files, const FString& DestinationPath, UFactory* ChosenFactory = NULL, bool bSyncToBrowser = true) const = 0;
 
 	/** Creates a unique package and asset name taking the form InBasePackageName+InSuffix */
 	virtual void CreateUniqueAssetName(const FString& InBasePackageName, const FString& InSuffix, FString& OutPackageName, FString& OutAssetName) const = 0;
diff --git a/Engine/Source/Editor/UnrealEd/Private/FeaturePack/FeaturePack.cpp b/Engine/Source/Editor/UnrealEd/Private/FeaturePack/FeaturePack.cpp
--- a/Engine/Source/Editor/UnrealEd/Private/FeaturePack/FeaturePack.cpp
+++ b/Engine/Source/Editor/UnrealEd/Private/FeaturePack/FeaturePack.cpp
@@ -70,7 +70,7 @@ void FFeaturePack::ParsePacks()
 			TArray<FString> EachImport;
 			FString FullPath = FPaths::FeaturePackDir() + EachPackData.PackSource;
 			EachImport.Add(FullPath);
-			EachPackData.ImportedObjects = AssetToolsModule.Get().ImportAssets(EachImport, EachPackData.PackName, nullptr);
+			EachPackData.ImportedObjects = AssetToolsModule.Get().ImportAssets(EachImport, EachPackData.PackName, nullptr, false);
 						
 			if (EachPackData.ImportedObjects.Num() == 0)
 			{			
