Added the ability to derive a Blueprint from a C++ class via the Content Browser

UE-9293 - Should offer 'Create BP' when right clicking on C++ class in Browser

There is now a "Create Blueprint class based on X" option alongside the "Create C++ class derived from X" option when viewing the context menu for a C++ Class in the Content Browser. This allows you to create a new Blueprint asset from C++ classes that are marked as blueprintable.

diff --git a/Engine/Source/Developer/AssetTools/Private/AssetTypeActions/AssetTypeActions_Class.cpp b/Engine/Source/Developer/AssetTools/Private/AssetTypeActions/AssetTypeActions_Class.cpp
--- a/Engine/Source/Developer/AssetTools/Private/AssetTypeActions/AssetTypeActions_Class.cpp
+++ b/Engine/Source/Developer/AssetTools/Private/AssetTypeActions/AssetTypeActions_Class.cpp
@@ -3,6 +3,7 @@
 #include "AssetToolsPrivatePCH.h"
 #include "SourceCodeNavigation.h"
 #include "GameProjectGenerationModule.h"
+#include "KismetEditorUtilities.h"
 
 #define LOCTEXT_NAMESPACE "AssetTypeActions"
 
@@ -18,9 +19,10 @@ void FAssetTypeActions_Class::GetActions(const TArray<UObject*>& InObjects, FMen
 	// Only allow the New class option if we have a base class that we can actually derive from in one of our project modules
 	FGameProjectGenerationModule& GameProjectGenerationModule = FGameProjectGenerationModule::Get();
 	TArray<FModuleContextInfo> ProjectModules = GameProjectGenerationModule.GetCurrentProjectModules();
-	const bool bIsValidBaseClass = BaseClass && GameProjectGenerationModule.IsValidBaseClassForCreation(BaseClass, ProjectModules);
+	const bool bIsValidBaseCppClass = BaseClass && GameProjectGenerationModule.IsValidBaseClassForCreation(BaseClass, ProjectModules);
+	const bool bIsValidBaseBlueprintClass = BaseClass && FKismetEditorUtilities::CanCreateBlueprintOfClass(BaseClass);
 
-	auto CreateCreateDerivedClass = [BaseClass]()
+	auto CreateCreateDerivedCppClass = [BaseClass]()
 	{
 		// Work out where the header file for the current class is, as we'll use that path as the default for the new class
 		FString BaseClassPath;
@@ -38,39 +40,88 @@ void FAssetTypeActions_Class::GetActions(const TArray<UObject*>& InObjects, FMen
 		);
 	};
 
-	auto CanCreateDerivedClass = [bIsValidBaseClass]() -> bool
+	auto CanCreateDerivedCppClass = [bIsValidBaseCppClass]() -> bool
 	{
-		return bIsValidBaseClass;
+		return bIsValidBaseCppClass;
 	};
 
-	FText NewDerivedClassClassLabel;
-	FText NewDerivedClassClassToolTip;
+	auto CreateCreateDerivedBlueprintClass = [BaseClass]()
+	{
+		// Work out where the header file for the current class is, as we'll use that path as the default for the new class
+		FString BaseClassPath;
+		if(FSourceCodeNavigation::FindClassHeaderPath(BaseClass, BaseClassPath))
+		{
+			// Strip off the actual filename as we only need the path
+			BaseClassPath = FPaths::GetPath(BaseClassPath);
+		}
+
+		FGameProjectGenerationModule::Get().OpenAddBlueprintToProjectDialog(
+			FAddToProjectConfig()
+			.ParentClass(BaseClass)
+			.InitialPath(BaseClassPath)
+			.ParentWindow(FGlobalTabmanager::Get()->GetRootWindow())
+		);
+	};
+
+	auto CanCreateDerivedBlueprintClass = [bIsValidBaseBlueprintClass]() -> bool
+	{
+		return bIsValidBaseBlueprintClass;
+	};
+
+	FText NewDerivedCppClassLabel;
+	FText NewDerivedCppClassToolTip;
+	FText NewDerivedBlueprintClassLabel;
+	FText NewDerivedBlueprintClassToolTip;
 	if(InObjects.Num() == 1)
 	{
 		const FText BaseClassName = FText::FromName(BaseClass->GetFName());
-		NewDerivedClassClassLabel = FText::Format(LOCTEXT("Class_NewDerivedClassLabel_CreateFrom", "Create class derived from {0}"), BaseClassName);
-		if(bIsValidBaseClass)
+
+		NewDerivedCppClassLabel = FText::Format(LOCTEXT("Class_NewDerivedCppClassLabel_CreateFrom", "Create C++ class derived from {0}"), BaseClassName);
+		if(bIsValidBaseCppClass)
+		{
+			NewDerivedCppClassToolTip = FText::Format(LOCTEXT("Class_NewDerivedCppClassTooltip_CreateFrom", "Create a new C++ class deriving from {0}."), BaseClassName);
+		}
+		else
+		{
+			NewDerivedCppClassToolTip = FText::Format(LOCTEXT("Class_NewDerivedCppClassTooltip_InvalidClass", "Cannot create a new C++ class deriving from {0}."), BaseClassName);
+		}
+
+		NewDerivedBlueprintClassLabel = FText::Format(LOCTEXT("Class_NewDerivedBlueprintClassLabel_CreateFrom", "Create Blueprint class based on {0}"), BaseClassName);
+		if(bIsValidBaseBlueprintClass)
 		{
-			NewDerivedClassClassToolTip = FText::Format(LOCTEXT("Class_NewDerivedClassTooltip_CreateFrom", "Create a new class deriving from {0}."), BaseClassName);
+			NewDerivedBlueprintClassToolTip = FText::Format(LOCTEXT("Class_NewDerivedBlueprintClassTooltip_CreateFrom", "Create a new Blueprint class based on {0}."), BaseClassName);
 		}
 		else
 		{
-			NewDerivedClassClassToolTip = FText::Format(LOCTEXT("Class_NewDerivedClassTooltip_InvalidClass", "Cannot create a new class deriving from {0}."), BaseClassName);
+			NewDerivedBlueprintClassToolTip = FText::Format(LOCTEXT("Class_NewDerivedBlueprintClassTooltip_InvalidClass", "Cannot create a new Blueprint class based on {0}."), BaseClassName);
 		}
 	}
 	else
 	{
-		NewDerivedClassClassLabel = LOCTEXT("Class_NewDerivedClassLabel_InvalidNumberOfBases", "New class derived from...");
-		NewDerivedClassClassToolTip = LOCTEXT("Class_NewDerivedClassTooltip_InvalidNumberOfBases", "Can only create a derived class when there is a single base class selected.");
+		NewDerivedCppClassLabel = LOCTEXT("Class_NewDerivedCppClassLabel_InvalidNumberOfBases", "New C++ class derived from...");
+		NewDerivedCppClassToolTip = LOCTEXT("Class_NewDerivedCppClassTooltip_InvalidNumberOfBases", "Can only create a derived C++ class when there is a single base class selected.");
+
+		NewDerivedBlueprintClassLabel = LOCTEXT("Class_NewDerivedBlueprintClassLabel_InvalidNumberOfBases", "New Blueprint class based on...");
+		NewDerivedBlueprintClassToolTip = LOCTEXT("Class_NewDerivedBlueprintClassTooltip_InvalidNumberOfBases", "Can only create a Blueprint class when there is a single base class selected.");
 	}
 
 	MenuBuilder.AddMenuEntry(
-		NewDerivedClassClassLabel,
-		NewDerivedClassClassToolTip,
+		NewDerivedCppClassLabel,
+		NewDerivedCppClassToolTip,
 		FSlateIcon(FEditorStyle::GetStyleSetName(), "MainFrame.AddCodeToProject"),
 		FUIAction(
-			FExecuteAction::CreateLambda(CreateCreateDerivedClass),
-			FCanExecuteAction::CreateLambda(CanCreateDerivedClass)
+			FExecuteAction::CreateLambda(CreateCreateDerivedCppClass),
+			FCanExecuteAction::CreateLambda(CanCreateDerivedCppClass)
+			)
+		);
+
+	MenuBuilder.AddMenuEntry(
+		NewDerivedBlueprintClassLabel,
+		NewDerivedBlueprintClassToolTip,
+		FSlateIcon(FEditorStyle::GetStyleSetName(), "LevelEditor.CreateClassBlueprint"),
+		FUIAction(
+			FExecuteAction::CreateLambda(CreateCreateDerivedBlueprintClass),
+			FCanExecuteAction::CreateLambda(CanCreateDerivedBlueprintClass)
 			)
 		);
 }
diff --git a/Engine/Source/Editor/GameProjectGeneration/Private/GameProjectUtils.cpp b/Engine/Source/Editor/GameProjectGeneration/Private/GameProjectUtils.cpp
--- a/Engine/Source/Editor/GameProjectGeneration/Private/GameProjectUtils.cpp
+++ b/Engine/Source/Editor/GameProjectGeneration/Private/GameProjectUtils.cpp
@@ -680,7 +680,7 @@ bool GameProjectUtils::UpdateGameProject(const FString& ProjectFile, const FStri
 void GameProjectUtils::OpenAddToProjectDialog(const FAddToProjectConfig& Config, EClassDomain InDomain)
 {
 	// If we've been given a class then we only show the second page of the dialog, so we can make the window smaller as that page doesn't have as much content
-	const FVector2D WindowSize = (Config._ParentClass) ? FVector2D(940, 380) : FVector2D(940, 540);
+	const FVector2D WindowSize = (Config._ParentClass) ? (InDomain == EClassDomain::Blueprint) ? FVector2D(940, 480) : FVector2D(940, 380) : FVector2D(940, 540);
 
 	FText WindowTitle = Config._WindowTitle;
 	if (WindowTitle.IsEmpty())
