[INTEGRATE] Change 2443487 by [REDACTED_PERSON][REDACTED_HANDLE] on 2015/02/12 10:49:16

	Changed c++ thumbnail color

diff --git a/Engine/Source/Developer/AssetTools/Private/AssetTypeActions/AssetTypeActions_Class.h b/Engine/Source/Developer/AssetTools/Private/AssetTypeActions/AssetTypeActions_Class.h
--- a/Engine/Source/Developer/AssetTools/Private/AssetTypeActions/AssetTypeActions_Class.h
+++ b/Engine/Source/Developer/AssetTools/Private/AssetTypeActions/AssetTypeActions_Class.h
@@ -9,7 +9,7 @@ class FAssetTypeActions_Class : public FAssetTypeActions_ClassTypeBase
 public:
 	// IAssetTypeActions Implementation
 	virtual FText GetName() const override { return NSLOCTEXT("AssetTypeActions", "AssetTypeActions_Class", "C++ Class"); }
-	virtual FColor GetTypeColor() const override { return FColor(72, 155, 129); }
+	virtual FColor GetTypeColor() const override { return FColor(255, 255, 255); }
 	virtual UClass* GetSupportedClass() const override { return UClass::StaticClass(); }
 	virtual uint32 GetCategories() override { return EAssetTypeCategories::Basic; }
 	
