Unified some terminology between the CB and SCS editor

UE-9405 - Update Blueprint option txt and tooltip to match SCS create child option

diff --git a/Engine/Source/Developer/AssetTools/Private/AssetTypeActions/AssetTypeActions_Blueprint.cpp b/Engine/Source/Developer/AssetTools/Private/AssetTypeActions/AssetTypeActions_Blueprint.cpp
--- a/Engine/Source/Developer/AssetTools/Private/AssetTypeActions/AssetTypeActions_Blueprint.cpp
+++ b/Engine/Source/Developer/AssetTools/Private/AssetTypeActions/AssetTypeActions_Blueprint.cpp
@@ -52,7 +52,7 @@ void FAssetTypeActions_Blueprint::GetActions( const TArray<UObject*>& InObjects,
 		TAttribute<FText> DynamicTooltipAttribute = TAttribute<FText>::Create(DynamicTooltipGetter);
 
 		MenuBuilder.AddMenuEntry(
-			LOCTEXT("Blueprint_NewDerivedBlueprint", "Create Blueprint based on this"),
+			LOCTEXT("Blueprint_NewDerivedBlueprint", "Create Child Blueprint Class"),
 			DynamicTooltipAttribute,
 			FSlateIcon(FEditorStyle::GetStyleSetName(), "LevelEditor.CreateClassBlueprint"),
 			FUIAction(
@@ -206,7 +206,7 @@ FText FAssetTypeActions_Blueprint::GetNewDerivedBlueprintTooltip(TWeakObjectPtr<
 	}
 	else
 	{
-		return LOCTEXT("Blueprint_NewDerivedBlueprintTooltip", "Creates a blueprint based on the selected blueprint.");
+		return LOCTEXT("Blueprint_NewDerivedBlueprintTooltip", "Creates a Child Blueprint Class based on the current Blueprint, allowing you to create variants easily.");
 	}
 }
 
diff --git a/Engine/Source/Editor/Blutility/Private/AssetTypeActions_EditorUtilityBlueprint.cpp b/Engine/Source/Editor/Blutility/Private/AssetTypeActions_EditorUtilityBlueprint.cpp
--- a/Engine/Source/Editor/Blutility/Private/AssetTypeActions_EditorUtilityBlueprint.cpp
+++ b/Engine/Source/Editor/Blutility/Private/AssetTypeActions_EditorUtilityBlueprint.cpp
@@ -88,7 +88,7 @@ void FAssetTypeActions_EditorUtilityBlueprint::GetActions(const TArray<UObject*>
 		TAttribute<FText> DynamicTooltipAttribute = TAttribute<FText>::Create(DynamicTooltipGetter);
 
 		MenuBuilder.AddMenuEntry(
-			LOCTEXT("Blutility_NewDerivedBlueprint", "Create Blueprint based on this"),
+			LOCTEXT("Blueprint_NewDerivedBlueprint", "Create Child Blueprint Class"),
 			DynamicTooltipAttribute,
 			FSlateIcon(),
 			FUIAction(
@@ -206,7 +206,7 @@ FText FAssetTypeActions_EditorUtilityBlueprint::GetNewDerivedBlueprintTooltip(TW
 	}
 	else
 	{
-		return LOCTEXT("Blutility_NewDerivedBlueprintTooltip", "Creates a blueprint based on the selected blueprint.");
+		return LOCTEXT("Blutility_NewDerivedBlueprintTooltip", "Creates a Child Blueprint Class based on the current Blueprint, allowing you to create variants easily.");
 	}
 }
 
