Fixed crash creating a blueprint from a CPP if no destination folder is selected

diff --git a/Engine/Source/Developer/AssetTools/Private/AssetTypeActions/AssetTypeActions_Class.cpp b/Engine/Source/Developer/AssetTools/Private/AssetTypeActions/AssetTypeActions_Class.cpp
--- a/Engine/Source/Developer/AssetTools/Private/AssetTypeActions/AssetTypeActions_Class.cpp
+++ b/Engine/Source/Developer/AssetTools/Private/AssetTypeActions/AssetTypeActions_Class.cpp
@@ -47,18 +47,9 @@ void FAssetTypeActions_Class::GetActions(const TArray<UObject*>& InObjects, FMen
 
 	auto CreateCreateDerivedBlueprintClass = [BaseClass]()
 	{
-		// Work out where the header file for the current class is, as we'll use that path as the default for the new class
-		FString BaseClassPath;
-		if(FSourceCodeNavigation::FindClassHeaderPath(BaseClass, BaseClassPath))
-		{
-			// Strip off the actual filename as we only need the path
-			BaseClassPath = FPaths::GetPath(BaseClassPath);
-		}
-
 		FGameProjectGenerationModule::Get().OpenAddBlueprintToProjectDialog(
 			FAddToProjectConfig()
 			.ParentClass(BaseClass)
-			.InitialPath(BaseClassPath)
 			.ParentWindow(FGlobalTabmanager::Get()->GetRootWindow())
 		);
 	};
