Fixed double-slash error in the asset rename manager when it was supplied with a package path with a trailing slash

diff --git a/Engine/Source/Developer/AssetTools/Private/AssetRenameManager.cpp b/Engine/Source/Developer/AssetTools/Private/AssetRenameManager.cpp
--- a/Engine/Source/Developer/AssetTools/Private/AssetRenameManager.cpp
+++ b/Engine/Source/Developer/AssetTools/Private/AssetRenameManager.cpp
@@ -590,7 +590,7 @@ void FAssetRenameManager::PerformAssetRename(TArray<FAssetRenameDataWithReferenc
 		ObjectTools::FPackageGroupName PGN;
 		PGN.ObjectName = RenameData.NewName;
 		PGN.GroupName = TEXT("");
-		PGN.PackageName = RenameData.PackagePath + TEXT("/") + PGN.ObjectName;
+		PGN.PackageName = RenameData.PackagePath / PGN.ObjectName;
 		const bool bLeaveRedirector = RenameData.bCreateRedirector;
 
 		UPackage* OldPackage = Asset->GetOutermost();
