[UETOOL-257] - Added an AssetTypeActions for the procedural foliage spawner

diff --git a/Engine/Source/Developer/AssetTools/Private/AssetTools.cpp b/Engine/Source/Developer/AssetTools/Private/AssetTools.cpp
--- a/Engine/Source/Developer/AssetTools/Private/AssetTools.cpp
+++ b/Engine/Source/Developer/AssetTools/Private/AssetTools.cpp
@@ -64,6 +64,7 @@ FAssetTools::FAssetTools()
 	RegisterAssetTypeActions( MakeShareable(new FAssetTypeActions_ParticleSystem) );
 	RegisterAssetTypeActions( MakeShareable(new FAssetTypeActions_PhysicalMaterial) );
 	RegisterAssetTypeActions( MakeShareable(new FAssetTypeActions_PhysicsAsset) );
+	RegisterAssetTypeActions( MakeShareable(new FAssetTypeActions_ProceduralFoliageSpawner) );
 	RegisterAssetTypeActions( MakeShareable(new FAssetTypeActions_Redirector) );
 	RegisterAssetTypeActions( MakeShareable(new FAssetTypeActions_ReverbEffect) );
 	RegisterAssetTypeActions( MakeShareable(new FAssetTypeActions_Rig) );
diff --git a/Engine/Source/Developer/AssetTools/Private/AssetToolsPrivatePCH.h b/Engine/Source/Developer/AssetTools/Private/AssetToolsPrivatePCH.h
--- a/Engine/Source/Developer/AssetTools/Private/AssetToolsPrivatePCH.h
+++ b/Engine/Source/Developer/AssetTools/Private/AssetToolsPrivatePCH.h
@@ -67,6 +67,7 @@ DECLARE_LOG_CATEGORY_EXTERN(LogAssetTools, Log, All);
 #include "AssetTypeActions/AssetTypeActions_ParticleSystem.h"
 #include "AssetTypeActions/AssetTypeActions_PhysicalMaterial.h"
 #include "AssetTypeActions/AssetTypeActions_PhysicsAsset.h"
+#include "AssetTypeActions/AssetTypeActions_ProceduralFoliageSpawner.h"
 #include "AssetTypeActions/AssetTypeActions_Redirector.h"
 #include "AssetTypeActions/AssetTypeActions_ReverbEffect.h"
 #include "AssetTypeActions/AssetTypeActions_Rig.h"
diff --git a/Engine/Source/Developer/AssetTools/Private/AssetTypeActions/AssetTypeActions_ProceduralFoliageSpawner.cpp b/Engine/Source/Developer/AssetTools/Private/AssetTypeActions/AssetTypeActions_ProceduralFoliageSpawner.cpp
new file mode 100644
--- /dev/null
+++ b/Engine/Source/Developer/AssetTools/Private/AssetTypeActions/AssetTypeActions_ProceduralFoliageSpawner.cpp
@@ -0,0 +1,9 @@
+// Copyright 1998-2015 Epic Games, Inc. All Rights Reserved.
+
+#include "AssetToolsPrivatePCH.h"
+#include "ProceduralFoliageSpawner.h"
+
+UClass* FAssetTypeActions_ProceduralFoliageSpawner::GetSupportedClass() const
+{
+	return UProceduralFoliageSpawner::StaticClass();
+}
\ No newline at end of file
diff --git a/Engine/Source/Developer/AssetTools/Private/AssetTypeActions/AssetTypeActions_ProceduralFoliageSpawner.h b/Engine/Source/Developer/AssetTools/Private/AssetTypeActions/AssetTypeActions_ProceduralFoliageSpawner.h
new file mode 100644
--- /dev/null
+++ b/Engine/Source/Developer/AssetTools/Private/AssetTypeActions/AssetTypeActions_ProceduralFoliageSpawner.h
@@ -0,0 +1,13 @@
+// Copyright 1998-2015 Epic Games, Inc. All Rights Reserved.
+
+#pragma once
+
+class FAssetTypeActions_ProceduralFoliageSpawner : public FAssetTypeActions_Base
+{
+public:
+	// IAssetTypeActions Implementation
+	virtual FText GetName() const override { return NSLOCTEXT("AssetTypeActions", "AssetTypeActions_ProceduralFoliageSpawner", "Procedural Foliage Spawner"); }
+	virtual FColor GetTypeColor() const override { return FColor(89, 134, 178); }
+	virtual UClass* GetSupportedClass() const override;
+	virtual uint32 GetCategories() override { return EAssetTypeCategories::Misc; }
+};
\ No newline at end of file
