Importing multiple assets with the same name no longer crashes the editor.
[UE-1102] - CRASH: Importing an FBX and TGA at the same time causes the editor to crash

diff --git a/Engine/Source/Developer/AssetTools/Private/AssetTools.cpp b/Engine/Source/Developer/AssetTools/Private/AssetTools.cpp
--- a/Engine/Source/Developer/AssetTools/Private/AssetTools.cpp
+++ b/Engine/Source/Developer/AssetTools/Private/AssetTools.cpp
@@ -723,102 +723,111 @@ TArray<UObject*> FAssetTools::ImportAssets(const TArray<FString>& Files, const F
 			UObject* ExistingObject = StaticFindObject( UObject::StaticClass(), Pkg, *Name );
 			if( ExistingObject != NULL )
 			{
-				// If the object is supported by the factory we are using, ask if we want to overwrite the asset
-				// Otherwise, prompt to replace the object
-				if ( Factory->DoesSupportClass(ExistingObject->GetClass()) )
+				// If the existing object is one of the imports we've just created we can't replace or overwrite it
+				if (ReturnObjects.Contains(ExistingObject))
 				{
-					// The factory can overwrite this object, ask if that is okay, unless "Yes To All" or "No To All" was already selected
-					EAppReturnType::Type UserResponse;
-					
-					if ( bOverwriteAll || GIsAutomationTesting)
-					{
-						UserResponse = EAppReturnType::YesAll;
-					}
-					else if ( bDontOverwriteAny )
-					{
-						UserResponse = EAppReturnType::NoAll;
-					}
-					else
-					{
-						UserResponse = FMessageDialog::Open(
-							EAppMsgType::YesNoYesAllNoAll,
-							FText::Format( LOCTEXT("ImportObjectAlreadyExists_SameClass", "Do you want to overwrite the existing asset?\n\nAn asset already exists at the import location: {0}"), FText::FromString( PackageName ) ) );
-
-						bOverwriteAll = UserResponse == EAppReturnType::YesAll;
-						bDontOverwriteAny = UserResponse == EAppReturnType::NoAll;
-					}
-
-					const bool bWantOverwrite = UserResponse == EAppReturnType::Yes || UserResponse == EAppReturnType::YesAll;
-
-					if( !bWantOverwrite )
-					{
-						// User chose not to replace the package
-						bImportWasCancelled = true;
-						OnNewImportRecord(ImportAssetType, FileExtension, bImportSucceeded, bImportWasCancelled, ImportStartTime);
-						continue;
-					}
+					// generate a unique name for this import
+					Name = MakeUniqueObjectName(Pkg, ImportAssetType, *Name).ToString();
 				}
 				else
 				{
-					// The factory can't overwrite this asset, ask if we should delete the object then import the new one. Only do this if "Yes To All" or "No To All" was not already selected.
-					EAppReturnType::Type UserResponse;
-					
-					if ( bReplaceAll )
+					// If the object is supported by the factory we are using, ask if we want to overwrite the asset
+					// Otherwise, prompt to replace the object
+					if (Factory->DoesSupportClass(ExistingObject->GetClass()))
 					{
-						UserResponse = EAppReturnType::YesAll;
-					}
-					else if ( bDontReplaceAny )
-					{
-						UserResponse = EAppReturnType::NoAll;
+						// The factory can overwrite this object, ask if that is okay, unless "Yes To All" or "No To All" was already selected
+						EAppReturnType::Type UserResponse;
+
+						if (bOverwriteAll || GIsAutomationTesting)
+						{
+							UserResponse = EAppReturnType::YesAll;
+						}
+						else if (bDontOverwriteAny)
+						{
+							UserResponse = EAppReturnType::NoAll;
+						}
+						else
+						{
+							UserResponse = FMessageDialog::Open(
+								EAppMsgType::YesNoYesAllNoAll,
+								FText::Format(LOCTEXT("ImportObjectAlreadyExists_SameClass", "Do you want to overwrite the existing asset?\n\nAn asset already exists at the import location: {0}"), FText::FromString(PackageName)));
+
+							bOverwriteAll = UserResponse == EAppReturnType::YesAll;
+							bDontOverwriteAny = UserResponse == EAppReturnType::NoAll;
+						}
+
+						const bool bWantOverwrite = UserResponse == EAppReturnType::Yes || UserResponse == EAppReturnType::YesAll;
+
+						if (!bWantOverwrite)
+						{
+							// User chose not to replace the package
+							bImportWasCancelled = true;
+							OnNewImportRecord(ImportAssetType, FileExtension, bImportSucceeded, bImportWasCancelled, ImportStartTime);
+							continue;
+						}
 					}
 					else
 					{
-						UserResponse = FMessageDialog::Open(
-							EAppMsgType::YesNoYesAllNoAll,
-							FText::Format( LOCTEXT("ImportObjectAlreadyExists_DifferentClass", "Do you want to replace the existing asset?\n\nAn asset already exists at the import location: {0}"), FText::FromString( PackageName ) ) );
+						// The factory can't overwrite this asset, ask if we should delete the object then import the new one. Only do this if "Yes To All" or "No To All" was not already selected.
+						EAppReturnType::Type UserResponse;
 
-						bReplaceAll = UserResponse == EAppReturnType::YesAll;
-						bDontReplaceAny = UserResponse == EAppReturnType::NoAll;
-					}
+						if (bReplaceAll)
+						{
+							UserResponse = EAppReturnType::YesAll;
+						}
+						else if (bDontReplaceAny)
+						{
+							UserResponse = EAppReturnType::NoAll;
+						}
+						else
+						{
+							UserResponse = FMessageDialog::Open(
+								EAppMsgType::YesNoYesAllNoAll,
+								FText::Format(LOCTEXT("ImportObjectAlreadyExists_DifferentClass", "Do you want to replace the existing asset?\n\nAn asset already exists at the import location: {0}"), FText::FromString(PackageName)));
 
-					const bool bWantReplace = UserResponse == EAppReturnType::Yes || UserResponse == EAppReturnType::YesAll;
+							bReplaceAll = UserResponse == EAppReturnType::YesAll;
+							bDontReplaceAny = UserResponse == EAppReturnType::NoAll;
+						}
 
-					if( bWantReplace )
-					{
-						// Delete the existing object
-						int32 NumObjectsDeleted = 0;
-						TArray< UObject* > ObjectsToDelete;
-						ObjectsToDelete.Add(ExistingObject);
-
-						// Dont let the package get garbage collected (just in case we are deleting the last asset in the package)
-						Pkg->AddToRoot();
-						NumObjectsDeleted = ObjectTools::DeleteObjects( ObjectsToDelete, /*bShowConfirmation=*/false );
-						Pkg->RemoveFromRoot();
-
-						const FString QualifiedName = PackageName + TEXT(".") + Name;
-						FText Reason;
-						if( NumObjectsDeleted == 0 || !IsUniqueObjectName( *QualifiedName, ANY_PACKAGE, Reason ) )
+						const bool bWantReplace = UserResponse == EAppReturnType::Yes || UserResponse == EAppReturnType::YesAll;
+
+						if (bWantReplace)
 						{
-							// Original object couldn't be deleted
-							const FText Message = FText::Format( LOCTEXT("ImportDeleteFailed", "Failed to delete '{0}'. The asset is referenced by other content."), FText::FromString( PackageName ) );
-							FMessageDialog::Open( EAppMsgType::Ok, Message );
-							UE_LOG(LogAssetTools, Warning, TEXT("%s"), *Message.ToString());
-							OnNewImportRecord(ImportAssetType, FileExtension, bImportSucceeded, bImportWasCancelled, ImportStartTime);
-							continue;
+							// Delete the existing object
+							int32 NumObjectsDeleted = 0;
+							TArray< UObject* > ObjectsToDelete;
+							ObjectsToDelete.Add(ExistingObject);
+
+							// Dont let the package get garbage collected (just in case we are deleting the last asset in the package)
+							Pkg->AddToRoot();
+							NumObjectsDeleted = ObjectTools::DeleteObjects(ObjectsToDelete, /*bShowConfirmation=*/false);
+							Pkg->RemoveFromRoot();
+
+							const FString QualifiedName = PackageName + TEXT(".") + Name;
+							FText Reason;
+							if (NumObjectsDeleted == 0 || !IsUniqueObjectName(*QualifiedName, ANY_PACKAGE, Reason))
+							{
+								// Original object couldn't be deleted
+								const FText Message = FText::Format(LOCTEXT("ImportDeleteFailed", "Failed to delete '{0}'. The asset is referenced by other content."), FText::FromString(PackageName));
+								FMessageDialog::Open(EAppMsgType::Ok, Message);
+								UE_LOG(LogAssetTools, Warning, TEXT("%s"), *Message.ToString());
+								OnNewImportRecord(ImportAssetType, FileExtension, bImportSucceeded, bImportWasCancelled, ImportStartTime);
+								continue;
+							}
+							else
+							{
+								// succeed, recreate package since it has been deleted
+								Pkg = CreatePackage(NULL, *PackageName);
+							}
 						}
 						else
 						{
-							// succeed, recreate package since it has been deleted
-							Pkg = CreatePackage(NULL, *PackageName);
+							// User chose not to replace the package
+							bImportWasCancelled = true;
+							OnNewImportRecord(ImportAssetType, FileExtension, bImportSucceeded, bImportWasCancelled, ImportStartTime);
+							continue;
 						}
 					}
-					else
-					{
-						// User chose not to replace the package
-						bImportWasCancelled = true;
-						OnNewImportRecord(ImportAssetType, FileExtension, bImportSucceeded, bImportWasCancelled, ImportStartTime);
-						continue;
-					}
 				}
 			}
 
