Editor: Added sanitization to the asset name in addition to the package name in FAssetTools::CreateUniqueAssetName

diff --git a/Engine/Source/Developer/AssetTools/Private/AssetTools.cpp b/Engine/Source/Developer/AssetTools/Private/AssetTools.cpp
--- a/Engine/Source/Developer/AssetTools/Private/AssetTools.cpp
+++ b/Engine/Source/Developer/AssetTools/Private/AssetTools.cpp
@@ -895,8 +895,12 @@ TArray<UObject*> FAssetTools::ImportAssets(const TArray<FString>& Files, const F
 
 void FAssetTools::CreateUniqueAssetName(const FString& InBasePackageName, const FString& InSuffix, FString& OutPackageName, FString& OutAssetName) const
 {
-	const FString PackagePath = FPackageName::GetLongPackagePath(InBasePackageName);
-	const FString BaseAssetName = FPackageName::GetLongPackageAssetName(InBasePackageName);
+	const FString SanitizedBasePackageName = PackageTools::SanitizePackageName(InBasePackageName);
+
+	const FString PackagePath = FPackageName::GetLongPackagePath(SanitizedBasePackageName);
+	const FString BaseAssetNameWithSuffix = FPackageName::GetLongPackageAssetName(SanitizedBasePackageName) + InSuffix;
+	const FString SanitizedBaseAssetName = ObjectTools::SanitizeObjectName(BaseAssetNameWithSuffix);
+
 	int32 IntSuffix = 1;
 	bool bObjectExists = false;
 
@@ -907,11 +911,11 @@ void FAssetTools::CreateUniqueAssetName(const FString& InBasePackageName, const
 		bObjectExists = false;
 		if ( IntSuffix <= 1 )
 		{
-			OutAssetName = FString::Printf( TEXT("%s%s"), *BaseAssetName, *InSuffix);
+			OutAssetName = SanitizedBaseAssetName;
 		}
 		else
 		{
-			OutAssetName = FString::Printf( TEXT("%s%s%d"), *BaseAssetName, *InSuffix, IntSuffix );
+			OutAssetName = FString::Printf(TEXT("%s%d"), *SanitizedBaseAssetName, IntSuffix);
 		}
 	
 		OutPackageName = PackagePath + TEXT("/") + OutAssetName;
