Fix issue where renamed assets were incorrectly absent from the "In Use By Level" filter in the Content Browser until toggling the filter on and off.

diff --git a/Engine/Source/Developer/AssetTools/Private/AssetRenameManager.cpp b/Engine/Source/Developer/AssetTools/Private/AssetRenameManager.cpp
--- a/Engine/Source/Developer/AssetTools/Private/AssetRenameManager.cpp
+++ b/Engine/Source/Developer/AssetTools/Private/AssetRenameManager.cpp
@@ -144,31 +144,32 @@ private:
 
 void FAssetRenameManager::RenameAssets(const TArray<FAssetRenameData>& AssetsAndNames) const
 {
-	// Prep a list of assets to rename with an extra boolean to determine if they should leave a redirector or not
-	TArray<FAssetRenameDataWithReferencers> AssetsToRename;
-	for ( int32 AssetIdx = 0; AssetIdx < AssetsAndNames.Num(); ++AssetIdx )
-	{
-		new(AssetsToRename) FAssetRenameDataWithReferencers(AssetsAndNames[AssetIdx]);
-	}
-
 	// If the asset registry is still loading assets, we cant check for referencers, so we must open the rename dialog
 	FAssetRegistryModule& AssetRegistryModule = FModuleManager::LoadModuleChecked<FAssetRegistryModule>("AssetRegistry");
 	if ( AssetRegistryModule.Get().IsLoadingAssets() )
 	{
 		// Open a dialog asking the user to wait while assets are being discovered
 		SDiscoveringAssetsDialog::OpenDiscoveringAssetsDialog(
-			SDiscoveringAssetsDialog::FOnAssetsDiscovered::CreateSP(this, &FAssetRenameManager::FixReferencesAndRename, AssetsToRename)
+			SDiscoveringAssetsDialog::FOnAssetsDiscovered::CreateSP(this, &FAssetRenameManager::FixReferencesAndRename, AssetsAndNames)
 			);
 	}
 	else
 	{
 		// No need to wait, attempt to fix references and rename now.
-		FixReferencesAndRename(AssetsToRename);
+		FixReferencesAndRename(AssetsAndNames);
 	}
 }
 
-void FAssetRenameManager::FixReferencesAndRename(TArray<FAssetRenameDataWithReferencers> AssetsToRename) const
+void FAssetRenameManager::FixReferencesAndRename(TArray<FAssetRenameData> AssetsAndNames) const
 {
+	// Prep a list of assets to rename with an extra boolean to determine if they should leave a redirector or not
+	TArray<FAssetRenameDataWithReferencers> AssetsToRename;
+	AssetsToRename.Reset(AssetsAndNames.Num());
+	for (int32 AssetIdx = 0; AssetIdx < AssetsAndNames.Num(); ++AssetIdx)
+	{
+		new(AssetsToRename)FAssetRenameDataWithReferencers(AssetsAndNames[AssetIdx]);
+	}
+
 	// Warn the user if they are about to rename an asset that is referenced by a CDO
 	auto CDOAssets = FindCDOReferencedAssets(AssetsToRename);
 
@@ -215,6 +216,9 @@ void FAssetRenameManager::FixReferencesAndRename(TArray<FAssetRenameDataWithRefe
 
 			// Save all packages that were referencing any of the assets that were moved without redirectors
 			SaveReferencingPackages(ReferencingPackagesToSave);
+
+			// Issue post rename event
+			AssetPostRenameEvent.Broadcast(AssetsAndNames);
 		}
 	}
 
diff --git a/Engine/Source/Developer/AssetTools/Private/AssetRenameManager.h b/Engine/Source/Developer/AssetTools/Private/AssetRenameManager.h
--- a/Engine/Source/Developer/AssetTools/Private/AssetRenameManager.h
+++ b/Engine/Source/Developer/AssetTools/Private/AssetRenameManager.h
@@ -19,6 +19,9 @@ public:
 	/** Renames assets using the specified names. */
 	void RenameAssets(const TArray<FAssetRenameData>& AssetsAndNames) const;
 
+	/** Accessor for post rename event */
+	FAssetPostRenameEvent& OnAssetPostRenameEvent() { return AssetPostRenameEvent; }
+
 	/**
 	 * Function that renames all FStringAssetReference object with the old asset path to the new one.
 	 *
@@ -30,7 +33,7 @@ public:
 
 private:
 	/** Attempts to load and fix redirector references for the supplied assets */
-	void FixReferencesAndRename(TArray<FAssetRenameDataWithReferencers> AssetsToRename) const;
+	void FixReferencesAndRename(TArray<FAssetRenameData> AssetsAndNames) const;
 
 	/** Get a list of assets referenced from CDOs */
 	TArray<TWeakObjectPtr<UObject>> FindCDOReferencedAssets(const TArray<FAssetRenameDataWithReferencers>& AssetsToRename) const;
@@ -65,4 +68,9 @@ private:
 
 	/** Report any failures that may have happened during the rename */
 	void ReportFailures(const TArray<FAssetRenameDataWithReferencers>& AssetsToRename) const;
+
+private:
+
+	/** Event issued at the end of the rename process */
+	FAssetPostRenameEvent AssetPostRenameEvent;
 };
\ No newline at end of file
diff --git a/Engine/Source/Developer/AssetTools/Private/AssetTools.h b/Engine/Source/Developer/AssetTools/Private/AssetTools.h
--- a/Engine/Source/Developer/AssetTools/Private/AssetTools.h
+++ b/Engine/Source/Developer/AssetTools/Private/AssetTools.h
@@ -32,6 +32,7 @@ public:
 	virtual bool CreateDiffProcess(const FString& DiffCommand, const FString& OldTextFilename, const FString& NewTextFilename, const FString& DiffArgs = FString("")) const override;
 	virtual void MigratePackages(const TArray<FName>& PackageNamesToMigrate) const override;
 	virtual void FixupReferencers(const TArray<UObjectRedirector*>& Objects) const override;
+	virtual FAssetPostRenameEvent& OnAssetPostRename() override { return AssetRenameManager->OnAssetPostRenameEvent(); }
 
 public:
 	/** Gets the asset tools singleton as a FAssetTools for asset tools module use */
diff --git a/Engine/Source/Developer/AssetTools/Public/IAssetTools.h b/Engine/Source/Developer/AssetTools/Public/IAssetTools.h
--- a/Engine/Source/Developer/AssetTools/Public/IAssetTools.h
+++ b/Engine/Source/Developer/AssetTools/Public/IAssetTools.h
@@ -26,6 +26,8 @@ struct FAssetRenameData
 	}
 };
 
+DECLARE_MULTICAST_DELEGATE_OneParam(FAssetPostRenameEvent, const TArray<FAssetRenameData>&);
+
 class IAssetTools
 {
 public:
@@ -87,6 +89,9 @@ public:
 	/** Renames assets using the specified names. */
 	virtual void RenameAssets(const TArray<FAssetRenameData>& AssetsAndNames) const = 0;
 
+	/** Event issued at the end of the rename process */
+	virtual FAssetPostRenameEvent& OnAssetPostRename() = 0;
+
 	/** Opens a file open dialog to choose files to import to the destination path. */
 	virtual TArray<UObject*> ImportAssets(const FString& DestinationPath) = 0;
 
diff --git a/Engine/Source/Editor/ContentBrowser/Private/FrontendFilters.cpp b/Engine/Source/Editor/ContentBrowser/Private/FrontendFilters.cpp
--- a/Engine/Source/Editor/ContentBrowser/Private/FrontendFilters.cpp
+++ b/Engine/Source/Editor/ContentBrowser/Private/FrontendFilters.cpp
@@ -4,6 +4,7 @@
 #include "FrontendFilters.h"
 #include "ISourceControlModule.h"
 #include "ObjectTools.h"
+#include "AssetToolsModule.h"
 
 
 /////////////////////////////////////////
@@ -271,11 +272,17 @@ FFrontendFilter_InUseByLoadedLevels::FFrontendFilter_InUseByLoadedLevels(TShared
 	, bIsCurrentlyActive(false)
 {
 	FEditorDelegates::MapChange.AddRaw(this, &FFrontendFilter_InUseByLoadedLevels::OnEditorMapChange);
+
+	IAssetTools& AssetTools = FAssetToolsModule::GetModule().Get();
+	AssetTools.OnAssetPostRename().AddRaw(this, &FFrontendFilter_InUseByLoadedLevels::OnAssetPostRename);
 }
 
 FFrontendFilter_InUseByLoadedLevels::~FFrontendFilter_InUseByLoadedLevels()
 {
 	FEditorDelegates::MapChange.RemoveAll(this);
+
+	IAssetTools& AssetTools = FAssetToolsModule::GetModule().Get();
+	AssetTools.OnAssetPostRename().RemoveAll(this);
 }
 
 void FFrontendFilter_InUseByLoadedLevels::ActiveStateChanged( bool bActive )
@@ -288,6 +295,12 @@ void FFrontendFilter_InUseByLoadedLevels::ActiveStateChanged( bool bActive )
 	}
 }
 
+void FFrontendFilter_InUseByLoadedLevels::OnAssetPostRename(const TArray<FAssetRenameData>& AssetsAndNames)
+{
+	// Update the tags identifying objects currently used by loaded levels
+	ObjectTools::TagInUseObjects(ObjectTools::SO_LoadedLevels);
+}
+
 bool FFrontendFilter_InUseByLoadedLevels::PassesFilter( AssetFilterType InItem ) const
 {
 	bool bObjectInUse = false;
diff --git a/Engine/Source/Editor/ContentBrowser/Private/FrontendFilters.h b/Engine/Source/Editor/ContentBrowser/Private/FrontendFilters.h
--- a/Engine/Source/Editor/ContentBrowser/Private/FrontendFilters.h
+++ b/Engine/Source/Editor/ContentBrowser/Private/FrontendFilters.h
@@ -175,6 +175,9 @@ public:
 	/** Handler for when maps change in the editor */
 	void OnEditorMapChange( uint32 MapChangeFlags );
 
+	/** Handler for when an asset is renamed */
+	void OnAssetPostRename(const TArray<FAssetRenameData>& AssetsAndNames);
+
 private:
 	bool bIsCurrentlyActive;
 };
