Audio: Removed the unused sound preview dialog and sound preview thread

diff --git a/Engine/Source/Developer/AssetTools/Private/AssetTypeActions/AssetTypeActions_SoundBase.cpp b/Engine/Source/Developer/AssetTools/Private/AssetTypeActions/AssetTypeActions_SoundBase.cpp
--- a/Engine/Source/Developer/AssetTools/Private/AssetTypeActions/AssetTypeActions_SoundBase.cpp
+++ b/Engine/Source/Developer/AssetTools/Private/AssetTypeActions/AssetTypeActions_SoundBase.cpp
@@ -1,7 +1,6 @@
 // Copyright 1998-2015 Epic Games, Inc. All Rights Reserved.
 
 #include "AssetToolsPrivatePCH.h"
-#include "Editor/UnrealEd/Public/Dialogs/DlgSoundWaveOptions.h"
 #include "Editor/SoundCueEditor/Public/SoundCueEditorModule.h"
 #include "SoundDefinitions.h"
 #include "Sound/SoundBase.h"
@@ -108,10 +107,7 @@ void FAssetTypeActions_SoundBase::PlaySound(USoundBase* Sound)
 {
 	if ( Sound )
 	{
-		if( !SSoundWaveCompressionOptions::IsQualityPreviewerActive() )
-		{
-			GEditor->PlayPreviewSound(Sound);
-		}
+		GEditor->PlayPreviewSound(Sound);
 	}
 	else
 	{
diff --git a/Engine/Source/Developer/AssetTools/Private/AssetTypeActions/AssetTypeActions_SoundCue.cpp b/Engine/Source/Developer/AssetTools/Private/AssetTypeActions/AssetTypeActions_SoundCue.cpp
--- a/Engine/Source/Developer/AssetTools/Private/AssetTypeActions/AssetTypeActions_SoundCue.cpp
+++ b/Engine/Source/Developer/AssetTools/Private/AssetTypeActions/AssetTypeActions_SoundCue.cpp
@@ -2,7 +2,6 @@
 
 #include "AssetToolsPrivatePCH.h"
 #include "AssetToolsModule.h"
-#include "Editor/UnrealEd/Public/Dialogs/DlgSoundWaveOptions.h"
 #include "Editor/SoundCueEditor/Public/SoundCueEditorModule.h"
 #include "SoundDefinitions.h"
 #include "Sound/SoundCue.h"
diff --git a/Engine/Source/Developer/AssetTools/Private/AssetTypeActions/AssetTypeActions_SoundWave.cpp b/Engine/Source/Developer/AssetTools/Private/AssetTypeActions/AssetTypeActions_SoundWave.cpp
--- a/Engine/Source/Developer/AssetTools/Private/AssetTypeActions/AssetTypeActions_SoundWave.cpp
+++ b/Engine/Source/Developer/AssetTools/Private/AssetTypeActions/AssetTypeActions_SoundWave.cpp
@@ -3,7 +3,6 @@
 #include "AssetToolsPrivatePCH.h"
 #include "AssetToolsModule.h"
 #include "SoundDefinitions.h"
-#include "Editor/UnrealEd/Public/Dialogs/DlgSoundWaveOptions.h"
 #include "ContentBrowserModule.h"
 #include "Sound/SoundWave.h"
 #include "Sound/SoundCue.h"
@@ -44,56 +43,6 @@ void FAssetTypeActions_SoundWave::GetResolvedSourceFilePaths(const TArray<UObjec
 void FAssetTypeActions_SoundWave::OpenAssetEditor( const TArray<UObject*>& InObjects, TSharedPtr<IToolkitHost> EditWithinLevelEditor )
 {
 	FSimpleAssetEditor::CreateEditor(EToolkitMode::Standalone, EditWithinLevelEditor, InObjects);
-
-/*	for (auto ObjIt = InObjects.CreateConstIterator(); ObjIt; ++ObjIt)
-	{
-		auto SoundWave = Cast<USoundWave>(*ObjIt);
-		if (SoundWave != NULL)
-		{
-			if( SoundWave->ChannelSizes.Num() == 0 && !SSoundWaveCompressionOptions::IsQualityPreviewerActive() && !GIsAutomationTesting )
-			{
-				TSharedPtr< SWindow > Window;
-				TSharedPtr< SSoundWaveCompressionOptions > Widget;
-
-				Window = SNew(SWindow)
-					.Title(NSLOCTEXT("SoundWaveOptions", "SoundWaveOptions_Title", "Sound Wave Compression Preview"))
-					.SizingRule( ESizingRule::Autosized )
-					.SupportsMaximize(false) .SupportsMinimize(false)
-					[
-						SNew( SBorder )
-						.Padding( 4.f )
-						.BorderImage( FEditorStyle::GetBrush( "ToolPanel.GroupBorder" ) )
-						[
-							SAssignNew(Widget, SSoundWaveCompressionOptions)
-								.SoundWave(SoundWave)
-						]
-					];
-
-				Widget->SetWindow(Window);
-
-				FSlateApplication::Get().AddWindow( Window.ToSharedRef() );
-			}
-			else if( SoundWave->ChannelSizes.Num() > 0 )
-			{
-				FMessageDialog::Open( EAppMsgType::Ok, LOCTEXT("SoundQualityPreviewerUnavailable", "Sound quality preview is unavailable for multichannel sounds.") );
-			}
-		}
-	}*/
-}
-
-void FAssetTypeActions_SoundWave::ExecuteCompressionPreview(TArray<TWeakObjectPtr<USoundWave>> Objects)
-{
-/*	TArray<UObject*> ObjectsToEdit;
-	for (auto ObjIt = Objects.CreateConstIterator(); ObjIt; ++ObjIt)
-	{
-		auto Object = (*ObjIt).Get();
-		if ( Object )
-		{
-			ObjectsToEdit.Add(Object);
-		}
-	}
-
-	OpenAssetEditor(ObjectsToEdit);*/
 }
 
 void FAssetTypeActions_SoundWave::ExecuteCreateSoundCue(TArray<TWeakObjectPtr<USoundWave>> Objects)
diff --git a/Engine/Source/Developer/AssetTools/Private/AssetTypeActions/AssetTypeActions_SoundWave.h b/Engine/Source/Developer/AssetTools/Private/AssetTypeActions/AssetTypeActions_SoundWave.h
--- a/Engine/Source/Developer/AssetTools/Private/AssetTypeActions/AssetTypeActions_SoundWave.h
+++ b/Engine/Source/Developer/AssetTools/Private/AssetTypeActions/AssetTypeActions_SoundWave.h
@@ -17,9 +17,6 @@ public:
 	
 private:
 
-	/** Handler for when CompressionPreview is selected */
-	void ExecuteCompressionPreview(TArray<TWeakObjectPtr<USoundWave>> Objects);
-
 	/** Creates a SoundCue of the same name for the sound, if one does not already exist */
 	void ExecuteCreateSoundCue(TArray<TWeakObjectPtr<USoundWave>> Objects);
 };
\ No newline at end of file
diff --git a/Engine/Source/Editor/UnrealEd/Private/Dialogs/DlgSoundWaveOptions.cpp b/Engine/Source/Editor/UnrealEd/Private/Dialogs/DlgSoundWaveOptions.cpp
deleted file mode 100644
diff --git a/Engine/Source/Editor/UnrealEd/Private/SoundPreviewThread.cpp b/Engine/Source/Editor/UnrealEd/Private/SoundPreviewThread.cpp
deleted file mode 100644
diff --git a/Engine/Source/Editor/UnrealEd/Public/Dialogs/DlgSoundWaveOptions.h b/Engine/Source/Editor/UnrealEd/Public/Dialogs/DlgSoundWaveOptions.h
deleted file mode 100644
diff --git a/Engine/Source/Editor/UnrealEd/Public/SoundPreviewThread.h b/Engine/Source/Editor/UnrealEd/Public/SoundPreviewThread.h
deleted file mode 100644
