Change to "Create Child Blueprint Class" so that the new asset is created and set up for inline renaming in the Content Browser, instead of opening the Blueprint Editor with the default name.

diff --git a/Engine/Source/Developer/AssetTools/Private/AssetTypeActions/AssetTypeActions_Blueprint.cpp b/Engine/Source/Developer/AssetTools/Private/AssetTypeActions/AssetTypeActions_Blueprint.cpp
--- a/Engine/Source/Developer/AssetTools/Private/AssetTypeActions/AssetTypeActions_Blueprint.cpp
+++ b/Engine/Source/Developer/AssetTools/Private/AssetTypeActions/AssetTypeActions_Blueprint.cpp
@@ -11,6 +11,7 @@
 #include "ISourceControlModule.h"
 #include "MessageLog.h"
 #include "Engine/BlueprintGeneratedClass.h"
+#include "ContentBrowserModule.h"
 
 #define LOCTEXT_NAMESPACE "AssetTypeActions"
 
@@ -175,26 +176,13 @@ void FAssetTypeActions_Blueprint::ExecuteNewDerivedBlueprint(TWeakObjectPtr<UBlu
 		FString Name;
 		FString PackageName;
 		CreateUniqueAssetName(Object->GetOutermost()->GetName(), TEXT("_Child"), PackageName, Name);
+		const FString PackagePath = FPackageName::GetLongPackagePath(PackageName);
 
-		UPackage* Package = CreatePackage(NULL, *PackageName);
-		if (ensure(Package))
-		{
-			// Create and init a new Blueprint
-			UBlueprint* NewBP = FKismetEditorUtilities::CreateBlueprint(TargetParentClass, Package, FName(*Name), BPTYPE_Normal, TargetParentBP->GetClass(), UBlueprintGeneratedClass::StaticClass());
-			if (NewBP)
-			{
-				// Notify the asset registry
-				FAssetRegistryModule::AssetCreated(NewBP);
-
-				// the editor should be opened AFTER being added to the asset 
-				// registry (some systems could queue off of the asset registry 
-				// add event, and already having that blueprint open can be odd)
-				FAssetEditorManager::Get().OpenEditorForAsset(NewBP);
+		UBlueprintFactory* BlueprintFactory = NewObject<UBlueprintFactory>();
+		BlueprintFactory->ParentClass = TargetParentClass;
 
-				// Mark the package dirty...
-				Package->MarkPackageDirty();
-			}
-		}
+		FContentBrowserModule& ContentBrowserModule = FModuleManager::LoadModuleChecked<FContentBrowserModule>("ContentBrowser");
+		ContentBrowserModule.Get().CreateNewAsset(Name, PackagePath, UBlueprint::StaticClass(), BlueprintFactory);
 	}
 }
 
diff --git a/Engine/Source/Editor/Blutility/Private/AssetTypeActions_EditorUtilityBlueprint.cpp b/Engine/Source/Editor/Blutility/Private/AssetTypeActions_EditorUtilityBlueprint.cpp
--- a/Engine/Source/Editor/Blutility/Private/AssetTypeActions_EditorUtilityBlueprint.cpp
+++ b/Engine/Source/Editor/Blutility/Private/AssetTypeActions_EditorUtilityBlueprint.cpp
@@ -10,6 +10,7 @@
 #include "AssetToolsModule.h"
 #include "GlobalBlutilityDialog.h"
 #include "Engine/BlueprintGeneratedClass.h"
+#include "ContentBrowserModule.h"
 
 #define LOCTEXT_NAMESPACE "AssetTypeActions"
 
@@ -179,22 +180,13 @@ void FAssetTypeActions_EditorUtilityBlueprint::ExecuteNewDerivedBlueprint(TWeakO
 		FString Name;
 		FString PackageName;
 		CreateUniqueAssetName(Object->GetOutermost()->GetName(), TEXT("_Child"), PackageName, Name);
+		const FString PackagePath = FPackageName::GetLongPackagePath(PackageName);
 
-		UPackage* Package = CreatePackage(NULL, *PackageName);
-		if (ensure(Package))
-		{
-			// Create and init a new Blueprint
-			if (UBlueprint* NewBP = FKismetEditorUtilities::CreateBlueprint(TargetClass, Package, FName(*Name), BPTYPE_Normal, UEditorUtilityBlueprint::StaticClass(), UBlueprintGeneratedClass::StaticClass()))
-			{
-				FAssetEditorManager::Get().OpenEditorForAsset(NewBP);
-
-				// Notify the asset registry
-				FAssetRegistryModule::AssetCreated(NewBP);
+		UEditorUtilityBlueprintFactory* BlueprintFactory = NewObject<UEditorUtilityBlueprintFactory>();
+		BlueprintFactory->ParentClass = TargetClass;
 
-				// Mark the package dirty...
-				Package->MarkPackageDirty();
-			}
-		}
+		FContentBrowserModule& ContentBrowserModule = FModuleManager::LoadModuleChecked<FContentBrowserModule>("ContentBrowser");
+		ContentBrowserModule.Get().CreateNewAsset(Name, PackagePath, UEditorUtilityBlueprint::StaticClass(), BlueprintFactory);
 	}
 }
 
