When creating unique asset name, avoid continually appending more numbers and instead increment them.

diff --git a/Engine/Source/Developer/AssetTools/Private/AssetTools.cpp b/Engine/Source/Developer/AssetTools/Private/AssetTools.cpp
--- a/Engine/Source/Developer/AssetTools/Private/AssetTools.cpp
+++ b/Engine/Source/Developer/AssetTools/Private/AssetTools.cpp
@@ -966,6 +966,20 @@ void FAssetTools::CreateUniqueAssetName(const FString& InBasePackageName, const
 	int32 IntSuffix = 1;
 	bool bObjectExists = false;
 
+	int32 CharIndex = SanitizedBaseAssetName.Len() - 1;
+	while (SanitizedBaseAssetName[CharIndex] >= TEXT('0') && SanitizedBaseAssetName[CharIndex] <= TEXT('9'))
+	{
+		--CharIndex;
+	}
+	FString TrailingInteger;
+	FString TrimmedBaseAssetName = SanitizedBaseAssetName;
+	if (CharIndex < SanitizedBaseAssetName.Len() - 1)
+	{
+		TrailingInteger = SanitizedBaseAssetName.RightChop(CharIndex + 1);
+		TrimmedBaseAssetName = SanitizedBaseAssetName.Left(CharIndex + 1);
+		IntSuffix = FCString::Atoi(*TrailingInteger);
+	}
+
 	FAssetRegistryModule& AssetRegistryModule = FModuleManager::Get().LoadModuleChecked<FAssetRegistryModule>(TEXT("AssetRegistry"));
 
 	do
@@ -977,7 +991,12 @@ void FAssetTools::CreateUniqueAssetName(const FString& InBasePackageName, const
 		}
 		else
 		{
-			OutAssetName = FString::Printf(TEXT("%s%d"), *SanitizedBaseAssetName, IntSuffix);
+			FString Suffix = FString::Printf(TEXT("%d"), IntSuffix);
+			while (Suffix.Len() < TrailingInteger.Len())
+			{
+				Suffix = TEXT("0") + Suffix;
+			}
+			OutAssetName = FString::Printf(TEXT("%s%s"), *TrimmedBaseAssetName, *Suffix);
 		}
 	
 		OutPackageName = PackagePath + TEXT("/") + OutAssetName;
