Fix crash when creating unique asset name if the source name was entirely numeric

diff --git a/Engine/Source/Developer/AssetTools/Private/AssetTools.cpp b/Engine/Source/Developer/AssetTools/Private/AssetTools.cpp
--- a/Engine/Source/Developer/AssetTools/Private/AssetTools.cpp
+++ b/Engine/Source/Developer/AssetTools/Private/AssetTools.cpp
@@ -967,13 +967,19 @@ void FAssetTools::CreateUniqueAssetName(const FString& InBasePackageName, const
 	bool bObjectExists = false;
 
 	int32 CharIndex = SanitizedBaseAssetName.Len() - 1;
-	while (SanitizedBaseAssetName[CharIndex] >= TEXT('0') && SanitizedBaseAssetName[CharIndex] <= TEXT('9'))
+	while (CharIndex >= 0 && SanitizedBaseAssetName[CharIndex] >= TEXT('0') && SanitizedBaseAssetName[CharIndex] <= TEXT('9'))
 	{
 		--CharIndex;
 	}
 	FString TrailingInteger;
 	FString TrimmedBaseAssetName = SanitizedBaseAssetName;
-	if (CharIndex < SanitizedBaseAssetName.Len() - 1)
+	if (SanitizedBaseAssetName.Len() > 0 && CharIndex == -1)
+	{
+		// This is the all numeric name, in this case we'd like to append _number, because just adding a number isn't great
+		TrimmedBaseAssetName += TEXT("_");
+		IntSuffix = 2;
+	}
+	if (CharIndex >= 0 && CharIndex < SanitizedBaseAssetName.Len() - 1)
 	{
 		TrailingInteger = SanitizedBaseAssetName.RightChop(CharIndex + 1);
 		TrimmedBaseAssetName = SanitizedBaseAssetName.Left(CharIndex + 1);
