Fixed issue where creating a child of an anim blueprint would result in a regular blueprint being created.  Also fixed a similar issue with GameplayAbilityBlueprints.

diff --git a/Engine/Source/Developer/AssetTools/Private/AssetTypeActions/AssetTypeActions_AnimBlueprint.cpp b/Engine/Source/Developer/AssetTools/Private/AssetTypeActions/AssetTypeActions_AnimBlueprint.cpp
--- a/Engine/Source/Developer/AssetTools/Private/AssetTypeActions/AssetTypeActions_AnimBlueprint.cpp
+++ b/Engine/Source/Developer/AssetTools/Private/AssetTypeActions/AssetTypeActions_AnimBlueprint.cpp
@@ -88,6 +88,15 @@ UThumbnailInfo* FAssetTypeActions_AnimBlueprint::GetThumbnailInfo(UObject* Asset
 	return ThumbnailInfo;
 }
 
+UFactory* FAssetTypeActions_AnimBlueprint::GetFactoryForBlueprintType(UBlueprint* InBlueprint) const
+{
+	UAnimBlueprintFactory* AnimBlueprintFactory = NewObject<UAnimBlueprintFactory>();
+	UAnimBlueprint* AnimBlueprint = CastChecked<UAnimBlueprint>(InBlueprint);
+	AnimBlueprintFactory->ParentClass = TSubclassOf<UAnimInstance>(*InBlueprint->GeneratedClass);
+	AnimBlueprintFactory->TargetSkeleton = AnimBlueprint->TargetSkeleton;
+	return AnimBlueprintFactory;
+}
+
 void FAssetTypeActions_AnimBlueprint::OpenAssetEditor( const TArray<UObject*>& InObjects, TSharedPtr<IToolkitHost> EditWithinLevelEditor )
 {
 	EToolkitMode::Type Mode = EditWithinLevelEditor.IsValid() ? EToolkitMode::WorldCentric : EToolkitMode::Standalone;
diff --git a/Engine/Source/Developer/AssetTools/Private/AssetTypeActions/AssetTypeActions_AnimBlueprint.h b/Engine/Source/Developer/AssetTools/Private/AssetTypeActions/AssetTypeActions_AnimBlueprint.h
--- a/Engine/Source/Developer/AssetTools/Private/AssetTypeActions/AssetTypeActions_AnimBlueprint.h
+++ b/Engine/Source/Developer/AssetTools/Private/AssetTypeActions/AssetTypeActions_AnimBlueprint.h
@@ -14,6 +14,9 @@ public:
 	virtual uint32 GetCategories() override { return EAssetTypeCategories::Animation; }
 	virtual class UThumbnailInfo* GetThumbnailInfo(UObject* Asset) const override;
 
+	// FAssetTypeActions_Blueprint interface
+	virtual UFactory* GetFactoryForBlueprintType(UBlueprint* InBlueprint) const override;
+
 private:
 
 	/** Handler to fill the retarget submenu */
diff --git a/Engine/Source/Developer/AssetTools/Private/AssetTypeActions/AssetTypeActions_Blueprint.cpp b/Engine/Source/Developer/AssetTools/Private/AssetTypeActions/AssetTypeActions_Blueprint.cpp
--- a/Engine/Source/Developer/AssetTools/Private/AssetTypeActions/AssetTypeActions_Blueprint.cpp
+++ b/Engine/Source/Developer/AssetTools/Private/AssetTypeActions/AssetTypeActions_Blueprint.cpp
@@ -122,6 +122,13 @@ bool FAssetTypeActions_Blueprint::CanCreateNewDerivedBlueprint() const
 	return true;
 }
 
+UFactory* FAssetTypeActions_Blueprint::GetFactoryForBlueprintType(UBlueprint* InBlueprint) const
+{
+	UBlueprintFactory* BlueprintFactory = NewObject<UBlueprintFactory>();
+	BlueprintFactory->ParentClass = InBlueprint->GeneratedClass;
+	return BlueprintFactory;
+}
+
 void FAssetTypeActions_Blueprint::ExecuteEditDefaults(TArray<TWeakObjectPtr<UBlueprint>> Objects)
 {
 	TArray< UBlueprint* > Blueprints;
@@ -178,11 +185,10 @@ void FAssetTypeActions_Blueprint::ExecuteNewDerivedBlueprint(TWeakObjectPtr<UBlu
 		CreateUniqueAssetName(Object->GetOutermost()->GetName(), TEXT("_Child"), PackageName, Name);
 		const FString PackagePath = FPackageName::GetLongPackagePath(PackageName);
 
-		UBlueprintFactory* BlueprintFactory = NewObject<UBlueprintFactory>();
-		BlueprintFactory->ParentClass = TargetParentClass;
+		UFactory* Factory = GetFactoryForBlueprintType(TargetParentBP);
 
 		FContentBrowserModule& ContentBrowserModule = FModuleManager::LoadModuleChecked<FContentBrowserModule>("ContentBrowser");
-		ContentBrowserModule.Get().CreateNewAsset(Name, PackagePath, UBlueprint::StaticClass(), BlueprintFactory);
+		ContentBrowserModule.Get().CreateNewAsset(Name, PackagePath, TargetParentBP->GetClass(), Factory);
 	}
 }
 
diff --git a/Engine/Source/Developer/AssetTools/Private/AssetTypeActions/AssetTypeActions_Blueprint.h b/Engine/Source/Developer/AssetTools/Private/AssetTypeActions/AssetTypeActions_Blueprint.h
--- a/Engine/Source/Developer/AssetTools/Private/AssetTypeActions/AssetTypeActions_Blueprint.h
+++ b/Engine/Source/Developer/AssetTools/Private/AssetTypeActions/AssetTypeActions_Blueprint.h
@@ -29,6 +29,9 @@ protected:
 	/** Whether or not this asset can create derived blueprints */
 	virtual bool CanCreateNewDerivedBlueprint() const;
 
+	/** Return the factory responsible for creating this type of Blueprint */
+	virtual UFactory* GetFactoryForBlueprintType(UBlueprint* InBlueprint) const;
+
 private:
 	/** Handler for when EditDefaults is selected */
 	void ExecuteEditDefaults(TArray<TWeakObjectPtr<UBlueprint>> Objects);
diff --git a/Engine/Source/Editor/GameplayAbilitiesEditor/Private/AssetTypeActions_GameplayAbilitiesBlueprint.cpp b/Engine/Source/Editor/GameplayAbilitiesEditor/Private/AssetTypeActions_GameplayAbilitiesBlueprint.cpp
--- a/Engine/Source/Editor/GameplayAbilitiesEditor/Private/AssetTypeActions_GameplayAbilitiesBlueprint.cpp
+++ b/Engine/Source/Editor/GameplayAbilitiesEditor/Private/AssetTypeActions_GameplayAbilitiesBlueprint.cpp
@@ -14,6 +14,7 @@
 #include "AssetTypeActions_GameplayAbilitiesBlueprint.h"
 #include "GameplayAbilitiesEditor.h"
 #include "GameplayAbilityBlueprint.h"
+#include "GameplayAbilitiesBlueprintFactory.h"
 
 #define LOCTEXT_NAMESPACE "AssetTypeActions"
 
@@ -64,4 +65,11 @@ UClass* FAssetTypeActions_GameplayAbilitiesBlueprint::GetSupportedClass() const
 	return UGameplayAbilityBlueprint::StaticClass(); 
 }
 
+UFactory* FAssetTypeActions_GameplayAbilitiesBlueprint::GetFactoryForBlueprintType(UBlueprint* InBlueprint) const
+{
+	UGameplayAbilitiesBlueprintFactory* GameplayAbilitiesBlueprintFactory = NewObject<UGameplayAbilitiesBlueprintFactory>();
+	GameplayAbilitiesBlueprintFactory->ParentClass = TSubclassOf<UGameplayAbility>(*InBlueprint->GeneratedClass);
+	return GameplayAbilitiesBlueprintFactory;
+}
+
 #undef LOCTEXT_NAMESPACE
\ No newline at end of file
diff --git a/Engine/Source/Editor/GameplayAbilitiesEditor/Private/AssetTypeActions_GameplayAbilitiesBlueprint.h b/Engine/Source/Editor/GameplayAbilitiesEditor/Private/AssetTypeActions_GameplayAbilitiesBlueprint.h
--- a/Engine/Source/Editor/GameplayAbilitiesEditor/Private/AssetTypeActions_GameplayAbilitiesBlueprint.h
+++ b/Engine/Source/Editor/GameplayAbilitiesEditor/Private/AssetTypeActions_GameplayAbilitiesBlueprint.h
@@ -15,6 +15,9 @@ public:
 	virtual uint32 GetCategories() override { return EAssetTypeCategories::Blueprint | EAssetTypeCategories::Gameplay; }
 	// End IAssetTypeActions Implementation
 
+	// FAssetTypeActions_Blueprint interface
+	virtual UFactory* GetFactoryForBlueprintType(UBlueprint* InBlueprint) const override;
+
 private:
 	/** Returns true if the blueprint is data only */
 	bool ShouldUseDataOnlyEditor(const UBlueprint* Blueprint) const;
