PR #1184: FIX Animation Sequence 'Reimport With new Source' context action (Contributed by [REDACTED_PERSON])

diff --git a/Engine/Source/Developer/AssetTools/Private/AssetTypeActions/AssetTypeActions_AnimSequence.cpp b/Engine/Source/Developer/AssetTools/Private/AssetTypeActions/AssetTypeActions_AnimSequence.cpp
--- a/Engine/Source/Developer/AssetTools/Private/AssetTypeActions/AssetTypeActions_AnimSequence.cpp
+++ b/Engine/Source/Developer/AssetTools/Private/AssetTypeActions/AssetTypeActions_AnimSequence.cpp
@@ -65,23 +65,30 @@ void FAssetTypeActions_AnimSequence::GetResolvedSourceFilePaths(const TArray<UOb
 
 void FAssetTypeActions_AnimSequence::ExecuteReimportWithNewSource(TArray<TWeakObjectPtr<UAnimSequence>> Objects)
 {
-	// clear the source, I'm not sure if I should just remove all AssetImportData, but this is only clearing the source
-	// so everything else stays except the path
-	for(auto ObjIt = Objects.CreateConstIterator(); ObjIt; ++ObjIt)
-	{
-		auto Object = (*ObjIt).Get();
-		if(Object && Object->AssetImportData)
-		{
-			Object->AssetImportData->SourceFilePath = TEXT("");
-		}
-	}
+	const FString EmptyPath = TEXT("");
 
 	for (auto ObjIt = Objects.CreateConstIterator(); ObjIt; ++ObjIt)
 	{
 		auto Object = (*ObjIt).Get();
 		if (Object)
 		{
-			FReimportManager::Instance()->Reimport(Object, /*bAskForNewFileIfMissing=*/true);
+			FString OldPath = EmptyPath;
+
+			// clear the source, I'm not sure if I should just remove all AssetImportData, but this is only clearing the source
+			// so everything else stays except the path
+			if (Object->AssetImportData)
+			{
+				OldPath = Object->AssetImportData->SourceFilePath;
+				Object->AssetImportData->SourceFilePath = EmptyPath;
+			}
+
+			bool bSuccess = FReimportManager::Instance()->Reimport(Object, /*bAskForNewFileIfMissing=*/true);
+
+			// restore the old source path in case reimport was not successful
+			if (!bSuccess && Object->AssetImportData && OldPath != EmptyPath)
+			{
+				Object->AssetImportData->SourceFilePath = OldPath;
+			}
 		}
 	}
 }
