Fix to Redirector Fixup so that redirectors are deleted cleanly and referencing map packages are loaded.

diff --git a/Engine/Source/Developer/AssetTools/Private/AssetFixUpRedirectors.cpp b/Engine/Source/Developer/AssetTools/Private/AssetFixUpRedirectors.cpp
--- a/Engine/Source/Developer/AssetTools/Private/AssetFixUpRedirectors.cpp
+++ b/Engine/Source/Developer/AssetTools/Private/AssetFixUpRedirectors.cpp
@@ -88,10 +88,14 @@ void FAssetFixUpRedirectors::ExecuteFixUp(TArray<TWeakObjectPtr<UObjectRedirecto
 				FixUpStringAssetReferences(RedirectorRefsList, ReferencingPackagesToSave);
 
 				// Save all packages that were referencing any of the assets that were moved without redirectors
-				SaveReferencingPackages(ReferencingPackagesToSave);
+				TArray<UPackage*> FailedToSave;
+				SaveReferencingPackages(ReferencingPackagesToSave, FailedToSave);
+
+				// Wait for package referencers to be updated
+				UpdateAssetReferencers(RedirectorRefsList);
 
 				// Delete any redirectors that are no longer referenced
-				DeleteRedirectors(RedirectorRefsList);
+				DeleteRedirectors(RedirectorRefsList, FailedToSave);
 
 				// Finally, report any failures that happened during the rename
 				ReportFailures(RedirectorRefsList);
@@ -172,15 +176,6 @@ void FAssetFixUpRedirectors::LoadReferencingPackages(TArray<FRedirectorRefs>& Re
 			UPackage* Package = FindPackage(NULL, *PackageName);
 			if ( !Package )
 			{
-				// Check if the package is a map before loading it!
-				if ( FEditorFileUtils::IsMapPackageAsset(PackageName) )
-				{
-					// This reference was a map package, don't load it
-					RedirectorRefs.bRedirectorValidForFixup = false;
-					RedirectorRefs.FailureReason = FText::Format(LOCTEXT("RedirectorFixupFailed_MapReference", "Redirector is referenced by an unloaded map. Package: {0}"), FText::FromString(PackageName));
-					continue;
-				}
-
 				Package = LoadPackage(NULL, *PackageName, LOAD_None);
 			}
 
@@ -289,19 +284,32 @@ void FAssetFixUpRedirectors::DetectReadOnlyPackages(TArray<FRedirectorRefs>& Red
 	}
 }
 
-void FAssetFixUpRedirectors::SaveReferencingPackages(const TArray<UPackage*>& ReferencingPackagesToSave) const
+void FAssetFixUpRedirectors::SaveReferencingPackages(const TArray<UPackage*>& ReferencingPackagesToSave, TArray<UPackage*>& OutFailedToSave) const
 {
 	if ( ReferencingPackagesToSave.Num() > 0 )
 	{
 		const bool bCheckDirty = false;
 		const bool bPromptToSave = false;
-		FEditorFileUtils::PromptForCheckoutAndSave(ReferencingPackagesToSave, bCheckDirty, bPromptToSave);
+		FEditorFileUtils::PromptForCheckoutAndSave(ReferencingPackagesToSave, bCheckDirty, bPromptToSave, &OutFailedToSave);
 
 		ISourceControlModule::Get().QueueStatusUpdate(ReferencingPackagesToSave);
 	}
 }
 
-void FAssetFixUpRedirectors::DeleteRedirectors(TArray<FRedirectorRefs>& RedirectorsToFix) const
+void FAssetFixUpRedirectors::UpdateAssetReferencers(const TArray<FRedirectorRefs>& RedirectorsToFix) const
+{
+	// Load the asset registry module
+	FAssetRegistryModule& AssetRegistryModule = FModuleManager::LoadModuleChecked<FAssetRegistryModule>(TEXT("AssetRegistry"));
+
+	TArray<FString> AssetPaths;
+	for (const auto& Redirector : RedirectorsToFix)
+	{
+		AssetPaths.AddUnique(FPackageName::GetLongPackagePath(Redirector.Redirector->GetOutermost()->GetName()));
+	}
+	AssetRegistryModule.Get().ScanPathsSynchronous(AssetPaths, true);
+}
+
+void FAssetFixUpRedirectors::DeleteRedirectors(TArray<FRedirectorRefs>& RedirectorsToFix, const TArray<UPackage*>& FailedToSave) const
 {
 	TArray<UObject*> ObjectsToDelete;
 	for ( auto RedirectorIt = RedirectorsToFix.CreateIterator(); RedirectorIt; ++RedirectorIt )
@@ -309,6 +317,21 @@ void FAssetFixUpRedirectors::DeleteRedirectors(TArray<FRedirectorRefs>& Redirect
 		FRedirectorRefs& RedirectorRefs = *RedirectorIt;
 		if ( RedirectorRefs.bRedirectorValidForFixup )
 		{
+			bool bAllReferencersFixedUp = true;
+			for (const auto& ReferencingPackageName : RedirectorRefs.ReferencingPackageNames)
+			{
+				if (FailedToSave.ContainsByPredicate([&](UPackage* Package) { return Package->GetFName() == ReferencingPackageName; }))
+				{
+					bAllReferencersFixedUp = false;
+					break;
+				}
+			}
+
+			if (!bAllReferencersFixedUp)
+			{
+				continue;
+			}
+
 			// Add all redirectors found in this package to the redirectors to delete list.
 			// All redirectors in this package should be fixed up.
 			UPackage* RedirectorPackage = RedirectorRefs.Redirector->GetOutermost();
@@ -352,7 +375,7 @@ void FAssetFixUpRedirectors::DeleteRedirectors(TArray<FRedirectorRefs>& Redirect
 
 	if ( ObjectsToDelete.Num() > 0 )
 	{
-		ObjectTools::DeleteObjects(ObjectsToDelete);
+		ObjectTools::DeleteObjects(ObjectsToDelete, false);
 	}
 }
 
diff --git a/Engine/Source/Developer/AssetTools/Private/AssetFixUpRedirectors.h b/Engine/Source/Developer/AssetTools/Private/AssetFixUpRedirectors.h
--- a/Engine/Source/Developer/AssetTools/Private/AssetFixUpRedirectors.h
+++ b/Engine/Source/Developer/AssetTools/Private/AssetFixUpRedirectors.h
@@ -41,10 +41,13 @@ private:
 	void FixUpStringAssetReferences(const TArray<FRedirectorRefs>& RedirectorsToFix, const TArray<UPackage*>& InReferencingPackagesToSave) const;
 
 	/** Saves all the referencing packages and updates SCC state */
-	void SaveReferencingPackages(const TArray<UPackage*>& ReferencingPackagesToSave) const;
+	void SaveReferencingPackages(const TArray<UPackage*>& ReferencingPackagesToSave, TArray<UPackage*>& OutFailedToSave) const;
+
+	/** Waits for the asset registry to update its asset referencer cache */
+	void UpdateAssetReferencers(const TArray<FRedirectorRefs>& RedirectorsToFix) const;
 
 	/** Deletes redirectors that are valid to delete */
-	void DeleteRedirectors(TArray<FRedirectorRefs>& RedirectorsToFix) const;
+	void DeleteRedirectors(TArray<FRedirectorRefs>& RedirectorsToFix, const TArray<UPackage*>& FailedToSave) const;
 
 	/** Report any failures that may have happened during the rename */
 	void ReportFailures(const TArray<FRedirectorRefs>& RedirectorsToFix) const;
