extracted EAssetTypeCategories into a separate header so that it can be cleanly used by other modules #UE4

diff --git a/Engine/Source/Developer/AssetTools/Public/AssetTypeCategories.h b/Engine/Source/Developer/AssetTools/Public/AssetTypeCategories.h
new file mode 100644
--- /dev/null
+++ b/Engine/Source/Developer/AssetTools/Public/AssetTypeCategories.h
@@ -0,0 +1,31 @@
+// Copyright 1998-2015 Epic Games, Inc. All Rights Reserved.
+
+#pragma once
+
+/**
+ *	this enum has been extracted into a separate header for easier inclusion in
+ *	code in other modules
+ */ 
+
+// The list of categories for Asset Type Actions and UFactory subclasses
+namespace EAssetTypeCategories
+{
+	enum Type
+	{
+		None = 0,
+		Basic = 1 << 0,
+		Animation = 1 << 1,
+		MaterialsAndTextures = 1 << 2,
+		Sounds = 1 << 3,
+		Physics = 1 << 4,
+		UI = 1 << 5,
+		Misc = 1 << 6,
+		Gameplay = 1 << 7,
+		Blueprint = 1 << 8,
+
+		// Items below this will be allocated at runtime via RegisterAdvancedAssetCategory
+		FirstUser = 1 << 9,
+		LastUser = 1 << 31,
+		// Last allowed value is 1 << 31
+	};
+}
\ No newline at end of file
diff --git a/Engine/Source/Developer/AssetTools/Public/IAssetTools.h b/Engine/Source/Developer/AssetTools/Public/IAssetTools.h
--- a/Engine/Source/Developer/AssetTools/Public/IAssetTools.h
+++ b/Engine/Source/Developer/AssetTools/Public/IAssetTools.h
@@ -2,6 +2,8 @@
 
 #pragma once
 
+#include "AssetTypeCategories.h"
+
 class IAssetTypeActions;
 class IClassTypeActions;
 class FAssetData;
@@ -28,30 +30,6 @@ struct FAssetRenameData
 
 DECLARE_MULTICAST_DELEGATE_OneParam(FAssetPostRenameEvent, const TArray<FAssetRenameData>&);
 
-// The list of categories for Asset Type Actions and UFactory subclasses
-namespace EAssetTypeCategories
-{
-	enum Type
-	{
-		None					= 0,
-		Basic					= 1 << 0,
-		Animation				= 1 << 1,
-		MaterialsAndTextures	= 1 << 2,
-		Sounds					= 1 << 3,
-		Physics					= 1 << 4,
-		UI						= 1 << 5,
-		Misc					= 1 << 6,
-		Gameplay				= 1 << 7,
-		Blueprint				= 1 << 8,
-
-		// Items below this will be allocated at runtime via RegisterAdvancedAssetCategory
-		FirstUser				= 1 << 9,
-		LastUser				= 1 << 31,
-		// Last allowed value is 1 << 31
-	};
-}
-
-
 struct FAdvancedAssetCategory
 {
 	EAssetTypeCategories::Type CategoryType;
diff --git a/Engine/Source/Runtime/AIModule/Public/AIModule.h b/Engine/Source/Runtime/AIModule/Public/AIModule.h
--- a/Engine/Source/Runtime/AIModule/Public/AIModule.h
+++ b/Engine/Source/Runtime/AIModule/Public/AIModule.h
@@ -4,7 +4,7 @@
 
 #include "AI/AISystemBase.h"
 #if WITH_EDITOR
-#include "Developer/AssetTools/Public/IAssetTools.h"
+#include "Developer/AssetTools/Public/AssetTypeCategories.h"
 #endif // WITH_EDITOR
 
 /**
