Animation: Added the AnimBP class thumbnail in the corner of the 3D thumnail for animation blueprint assets in the Content Browser

diff --git a/Engine/Source/Developer/AssetTools/Private/AssetTypeActions/AssetTypeActions_AnimBlueprint.cpp b/Engine/Source/Developer/AssetTools/Private/AssetTypeActions/AssetTypeActions_AnimBlueprint.cpp
--- a/Engine/Source/Developer/AssetTools/Private/AssetTypeActions/AssetTypeActions_AnimBlueprint.cpp
+++ b/Engine/Source/Developer/AssetTools/Private/AssetTypeActions/AssetTypeActions_AnimBlueprint.cpp
@@ -6,6 +6,7 @@
 #include "NotificationManager.h"
 #include "SNotificationList.h"
 #include "SSkeletonWidget.h"
+#include "ClassIconFinder.h"
 
 #define LOCTEXT_NAMESPACE "AssetTypeActions"
 
@@ -200,4 +201,20 @@ void FAssetTypeActions_AnimBlueprint::RetargetAssets(TArray<UObject*> InAnimBlue
 	SAnimationRemapSkeleton::ShowWindow(OldSkeleton, Message, FOnRetargetAnimation::CreateSP(this, &FAssetTypeActions_AnimBlueprint::RetargetAnimationHandler, bDuplicateAssets, AnimBlueprints));
 }
 
+TSharedPtr<SWidget> FAssetTypeActions_AnimBlueprint::GetThumbnailOverlay(const FAssetData& AssetData) const
+{
+	const FSlateBrush* Icon = FClassIconFinder::FindIconForClass(UAnimBlueprint::StaticClass());
+
+	return SNew(SBorder)
+		.BorderImage(FEditorStyle::GetNoBrush())
+		.Visibility(EVisibility::HitTestInvisible)
+		.Padding(FMargin(0.0f, 0.0f, 0.0f, 3.0f))
+		.HAlign(HAlign_Right)
+		.VAlign(VAlign_Bottom)
+		[
+			SNew(SImage)
+			.Image(Icon)
+		];
+}
+
 #undef LOCTEXT_NAMESPACE
\ No newline at end of file
diff --git a/Engine/Source/Developer/AssetTools/Private/AssetTypeActions/AssetTypeActions_AnimBlueprint.h b/Engine/Source/Developer/AssetTools/Private/AssetTypeActions/AssetTypeActions_AnimBlueprint.h
--- a/Engine/Source/Developer/AssetTools/Private/AssetTypeActions/AssetTypeActions_AnimBlueprint.h
+++ b/Engine/Source/Developer/AssetTools/Private/AssetTypeActions/AssetTypeActions_AnimBlueprint.h
@@ -13,6 +13,7 @@ public:
 	virtual void OpenAssetEditor( const TArray<UObject*>& InObjects, TSharedPtr<class IToolkitHost> EditWithinLevelEditor = TSharedPtr<IToolkitHost>() ) override;
 	virtual uint32 GetCategories() override { return EAssetTypeCategories::Animation; }
 	virtual class UThumbnailInfo* GetThumbnailInfo(UObject* Asset) const override;
+	virtual TSharedPtr<SWidget> GetThumbnailOverlay(const FAssetData& AssetData) const override;
 
 	// FAssetTypeActions_Blueprint interface
 	virtual UFactory* GetFactoryForBlueprintType(UBlueprint* InBlueprint) const override;
