Fixed issues with retargeting when source skeleton is missing or target skeleton is missing preview mesh

https://jira.ol.epicgames.net/browse/UE-14149
https://jira.ol.epicgames.net/browse/UE-4638
https://jira.ol.epicgames.net/browse/UE-4484

diff --git a/Engine/Source/Developer/AssetTools/Private/AssetTypeActions/AssetTypeActions_AnimBlueprint.cpp b/Engine/Source/Developer/AssetTools/Private/AssetTypeActions/AssetTypeActions_AnimBlueprint.cpp
--- a/Engine/Source/Developer/AssetTools/Private/AssetTypeActions/AssetTypeActions_AnimBlueprint.cpp
+++ b/Engine/Source/Developer/AssetTools/Private/AssetTypeActions/AssetTypeActions_AnimBlueprint.cpp
@@ -122,16 +122,11 @@ void FAssetTypeActions_AnimBlueprint::OpenAssetEditor( const TArray<UObject*>& I
 					RetargetAssets(AnimBlueprints, bDuplicateAssets);
 				}
 			}
-
-			if(AnimBlueprint->TargetSkeleton)
+			else
 			{
 				FPersonaModule& PersonaModule = FModuleManager::LoadModuleChecked<FPersonaModule>( "Persona" );
 				PersonaModule.CreatePersona( Mode, EditWithinLevelEditor, NULL, AnimBlueprint, NULL, NULL);
 			}
-			else
-			{
-				FMessageDialog::Open( EAppMsgType::Ok, LOCTEXT("FailedToLoadSkeletonlessAnimBlueprint", "The Anim Blueprint could not be loaded because it's skeleton is missing."));
-			}
 		}
 		else
 		{
@@ -164,7 +159,7 @@ void FAssetTypeActions_AnimBlueprint::ExecuteFindSkeleton(TArray<TWeakObjectPtr<
 
 void FAssetTypeActions_AnimBlueprint::RetargetAnimationHandler(USkeleton* OldSkeleton, USkeleton* NewSkeleton, bool bRemapReferencedAssets, bool bConvertSpaces, bool bDuplicateAssets, TArray<TWeakObjectPtr<UObject>> AnimBlueprints)
 {
-	if((!OldSkeleton || OldSkeleton->GetPreviewMesh(true)) && (!NewSkeleton || NewSkeleton->GetPreviewMesh(true)))
+	if(!OldSkeleton || OldSkeleton->GetPreviewMesh(true))
 	{
 		EditorAnimUtils::RetargetAnimations(OldSkeleton, NewSkeleton, AnimBlueprints, bRemapReferencedAssets, bDuplicateAssets, bConvertSpaces);
 	}
diff --git a/Engine/Source/Developer/AssetTools/Private/AssetTypeActions/AssetTypeActions_AnimationAsset.cpp b/Engine/Source/Developer/AssetTools/Private/AssetTypeActions/AssetTypeActions_AnimationAsset.cpp
--- a/Engine/Source/Developer/AssetTools/Private/AssetTypeActions/AssetTypeActions_AnimationAsset.cpp
+++ b/Engine/Source/Developer/AssetTools/Private/AssetTypeActions/AssetTypeActions_AnimationAsset.cpp
@@ -154,7 +154,7 @@ void FAssetTypeActions_AnimationAsset::ExecuteFindSkeleton(TArray<TWeakObjectPtr
 
 void FAssetTypeActions_AnimationAsset::RetargetAnimationHandler(USkeleton* OldSkeleton, USkeleton* NewSkeleton, bool bRemapReferencedAssets, bool bConvertSpaces, bool bDuplicateAssets,  TArray<TWeakObjectPtr<UObject>> InAnimAssets)
 {
-	if((!OldSkeleton || OldSkeleton->GetPreviewMesh(false)) && (!NewSkeleton || NewSkeleton->GetPreviewMesh(false)))
+	if((!OldSkeleton || OldSkeleton->GetPreviewMesh(false)))
 	{
 		EditorAnimUtils::RetargetAnimations(OldSkeleton, NewSkeleton, InAnimAssets, bRemapReferencedAssets, bDuplicateAssets, bConvertSpaces);
 	}
diff --git a/Engine/Source/Developer/AssetTools/Private/AssetTypeActions/AssetTypeActions_Skeleton.cpp b/Engine/Source/Developer/AssetTools/Private/AssetTypeActions/AssetTypeActions_Skeleton.cpp
--- a/Engine/Source/Developer/AssetTools/Private/AssetTypeActions/AssetTypeActions_Skeleton.cpp
+++ b/Engine/Source/Developer/AssetTools/Private/AssetTypeActions/AssetTypeActions_Skeleton.cpp
@@ -521,7 +521,7 @@ void FAssetTypeActions_Skeleton::CreateRig(const TWeakObjectPtr<USkeleton> Skele
 
 void FAssetTypeActions_Skeleton::RetargetAnimationHandler(USkeleton* OldSkeleton, USkeleton* NewSkeleton, bool bRemapReferencedAssets, bool bConvertSpaces)
 {
-	if((OldSkeleton && OldSkeleton->GetPreviewMesh(true) == NULL) || (NewSkeleton && NewSkeleton->GetPreviewMesh(true)==NULL))
+	if((OldSkeleton && OldSkeleton->GetPreviewMesh(true) == NULL))
 	{
 		FFormatNamedArguments Args;
 		Args.Add(TEXT("OldSkeletonName"), FText::FromString(GetNameSafe(OldSkeleton)));
diff --git a/Engine/Source/Editor/UnrealEd/Private/SSkeletonWidget.cpp b/Engine/Source/Editor/UnrealEd/Private/SSkeletonWidget.cpp
--- a/Engine/Source/Editor/UnrealEd/Private/SSkeletonWidget.cpp
+++ b/Engine/Source/Editor/UnrealEd/Private/SSkeletonWidget.cpp
@@ -719,25 +719,7 @@ FReply SAnimationRemapSkeleton::OnApply()
 {
 	if (OnRetargetAnimationDelegate.IsBound())
 	{
-		// verify if old skeleton and new skeleton has preview mesh, otherwise, it won't work. 
-		if ( (!OldSkeleton || OldSkeleton->GetPreviewMesh(true)) && (!NewSkeleton || NewSkeleton->GetPreviewMesh(true)))
-		{
-			OnRetargetAnimationDelegate.Execute(OldSkeleton, NewSkeleton, bRemapReferencedAssets, bConvertSpaces);
-		}
-		else
-		{
-			FFormatNamedArguments Args;
-			Args.Add(TEXT("OldSkeletonName"), FText::FromString(GetNameSafe(OldSkeleton)));
-			Args.Add(TEXT("NewSkeletonName"), FText::FromString(GetNameSafe(NewSkeleton)));
-			FNotificationInfo Info(FText::Format(LOCTEXT("Retarget Failed", "Old Skeleton {OldSkeletonName} and New Skeleton {NewSkeletonName} need to have Preview Mesh set up to convert animation"), Args));
-			Info.ExpireDuration = 5.0f;
-			Info.bUseLargeFont = false;
-			TSharedPtr<SNotificationItem> Notification = FSlateNotificationManager::Get().AddNotification(Info);
-			if(Notification.IsValid())
-			{
-				Notification->SetCompletionState(SNotificationItem::CS_Fail);
-			}
-		}
+		OnRetargetAnimationDelegate.Execute(OldSkeleton, NewSkeleton, bRemapReferencedAssets, bConvertSpaces);
 	}
 
 	CloseWindow();
diff --git a/Engine/Source/Runtime/Engine/Private/Animation/AnimSequence.cpp b/Engine/Source/Runtime/Engine/Private/Animation/AnimSequence.cpp
--- a/Engine/Source/Runtime/Engine/Private/Animation/AnimSequence.cpp
+++ b/Engine/Source/Runtime/Engine/Private/Animation/AnimSequence.cpp
@@ -2015,7 +2015,7 @@ void UAnimSequence::RemapTracksToNewSkeleton( USkeleton* NewSkeleton, bool bConv
 		USkeleton* OldSkeleton = GetSkeleton();
 
 		// first check if both has same rig, if so, we'll retarget using it
-		if (OldSkeleton && OldSkeleton->GetRig() != NULL && NewSkeleton->GetRig() == OldSkeleton->GetRig())
+		if (OldSkeleton && OldSkeleton->GetRig() != NULL && NewSkeleton->GetRig() == OldSkeleton->GetRig() && OldSkeleton->GetPreviewMesh() && NewSkeleton->GetPreviewMesh())
 		{
 			const URig* Rig = OldSkeleton->GetRig();
 
