Fix a crash when importing multiple objects with the same name at the same time when one of them is an FBX file.

https://jira.ol.epicgames.net/browse/UE-18687

diff --git a/Engine/Source/Developer/AssetTools/Private/AssetTools.cpp b/Engine/Source/Developer/AssetTools/Private/AssetTools.cpp
--- a/Engine/Source/Developer/AssetTools/Private/AssetTools.cpp
+++ b/Engine/Source/Developer/AssetTools/Private/AssetTools.cpp
@@ -795,6 +795,12 @@ TArray<UObject*> FAssetTools::ImportAssets(const TArray<FString>& Files, const F
 				// If the existing object is one of the imports we've just created we can't replace or overwrite it
 				if (ReturnObjects.Contains(ExistingObject))
 				{
+					if( ImportAssetType == nullptr )
+					{
+						// The factory probably supports multiple types and cant be determined yet without asking the user or actually loading it
+						// We just need to generate an unused name so object should do fine.
+						ImportAssetType = UObject::StaticClass();
+					}
 					// generate a unique name for this import
 					Name = MakeUniqueObjectName(Pkg, ImportAssetType, *Name).ToString();
 				}
