Fixing a crash after renaming a map or asset that couldn't have been renamed due to errors.

UE-18417

diff --git a/Engine/Source/Developer/AssetTools/Private/AssetRenameManager.cpp b/Engine/Source/Developer/AssetTools/Private/AssetRenameManager.cpp
--- a/Engine/Source/Developer/AssetTools/Private/AssetRenameManager.cpp
+++ b/Engine/Source/Developer/AssetTools/Private/AssetRenameManager.cpp
@@ -689,7 +689,12 @@ void FAssetRenameManager::PerformAssetRename(TArray<FAssetRenameDataWithReferenc
 	// Now branch the files in source control if possible
 	for (const auto& AssetToRename : AssetsToRename)
 	{
-		SourceControlHelpers::BranchPackage(AssetToRename.Asset->GetOutermost(), FindPackage(nullptr, *AssetToRename.OriginalAssetPath));
+		// If something went wrong when saving and the new asset does not exist on disk, don't branch it
+		// as it will just create a copy and any attempt to load it will result in crashes.
+		if (FPackageName::DoesPackageExist(AssetToRename.Asset->GetOutermost()->GetFName().ToString()))
+		{
+			SourceControlHelpers::BranchPackage(AssetToRename.Asset->GetOutermost(), FindPackage(nullptr, *AssetToRename.OriginalAssetPath));
+		}
 	}
 
 	// Clean up all packages that were left empty
