Experimental support for Widget and Anim blueprint diffing

diff --git a/Engine/Config/BaseEditorPerProjectUserSettings.ini b/Engine/Config/BaseEditorPerProjectUserSettings.ini
--- a/Engine/Config/BaseEditorPerProjectUserSettings.ini
+++ b/Engine/Config/BaseEditorPerProjectUserSettings.ini
@@ -131,7 +131,8 @@ bWorldBrowser=False
 bBreakOnExceptions=False
 bToolbarCustomization=False
 bBehaviorTreeEditor=False
-
+bEnableWidgetVisualDiff=False
+bEnableAnimVisualDiff=False
 
 [/Script/UnrealEd.EditorLoadingSavingSettings]
 ; True if we should automatically load a default level at start up
diff --git a/Engine/Source/Developer/AssetTools/Private/AssetTypeActions/AssetTypeActions_AnimBlueprint.cpp b/Engine/Source/Developer/AssetTools/Private/AssetTypeActions/AssetTypeActions_AnimBlueprint.cpp
--- a/Engine/Source/Developer/AssetTools/Private/AssetTypeActions/AssetTypeActions_AnimBlueprint.cpp
+++ b/Engine/Source/Developer/AssetTools/Private/AssetTypeActions/AssetTypeActions_AnimBlueprint.cpp
@@ -4,6 +4,7 @@
 #include "PersonaModule.h"
 #include "EditorAnimUtils.h"
 #include "NotificationManager.h"
+#include "SBlueprintDiff.h"
 #include "SNotificationList.h"
 #include "SSkeletonWidget.h"
 #include "ClassIconFinder.h"
@@ -135,6 +136,51 @@ void FAssetTypeActions_AnimBlueprint::OpenAssetEditor( const TArray<UObject*>& I
 	}
 }
 
+void FAssetTypeActions_AnimBlueprint::PerformAssetDiff(UObject* Asset1, UObject* Asset2, const struct FRevisionInfo& OldRevision, const struct FRevisionInfo& NewRevision) const
+{
+	if (!GetDefault<UEditorExperimentalSettings>()->bEnableAnimVisualDiff)
+	{
+		return FAssetTypeActions_Base::PerformAssetDiff(Asset1, Asset2, OldRevision, NewRevision);
+	}
+
+	UBlueprint* OldBlueprint = CastChecked<UBlueprint>(Asset1);
+	UBlueprint* NewBlueprint = CastChecked<UBlueprint>(Asset2);
+
+	// sometimes we're comparing different revisions of one single asset (other 
+	// times we're comparing two completely separate assets altogether)
+	bool bIsSingleAsset = (NewBlueprint->GetName() == OldBlueprint->GetName());
+
+	FText WindowTitle = LOCTEXT("NamelessAnimationBlueprintDiff", "Animation Blueprint Diff");
+	// if we're diff'ing one asset against itself 
+	if (bIsSingleAsset)
+	{
+		// identify the assumed single asset in the window's title
+		WindowTitle = FText::Format(LOCTEXT("AnimationBlueprintDiff", "{0} - Animation Blueprint Diff"), FText::FromString(NewBlueprint->GetName()));
+	}
+
+	const TSharedPtr<SWindow> Window = SNew(SWindow)
+		.Title(WindowTitle)
+		.ClientSize(FVector2D(1000, 800));
+
+	Window->SetContent(SNew(SBlueprintDiff)
+		.BlueprintOld(OldBlueprint)
+		.BlueprintNew(NewBlueprint)
+		.OldRevision(OldRevision)
+		.NewRevision(NewRevision)
+		.ShowAssetNames(!bIsSingleAsset));
+
+	// Make this window a child of the modal window if we've been spawned while one is active.
+	TSharedPtr<SWindow> ActiveModal = FSlateApplication::Get().GetActiveModalWindow();
+	if (ActiveModal.IsValid())
+	{
+		FSlateApplication::Get().AddWindowAsNativeChild(Window.ToSharedRef(), ActiveModal.ToSharedRef());
+	}
+	else
+	{
+		FSlateApplication::Get().AddWindow(Window.ToSharedRef());
+	}
+}
+
 void FAssetTypeActions_AnimBlueprint::ExecuteFindSkeleton(TArray<TWeakObjectPtr<UAnimBlueprint>> Objects)
 {
 	TArray<UObject*> ObjectsToSync;
diff --git a/Engine/Source/Developer/AssetTools/Private/AssetTypeActions/AssetTypeActions_AnimBlueprint.h b/Engine/Source/Developer/AssetTools/Private/AssetTypeActions/AssetTypeActions_AnimBlueprint.h
--- a/Engine/Source/Developer/AssetTools/Private/AssetTypeActions/AssetTypeActions_AnimBlueprint.h
+++ b/Engine/Source/Developer/AssetTools/Private/AssetTypeActions/AssetTypeActions_AnimBlueprint.h
@@ -12,6 +12,7 @@ public:
 	virtual void GetActions( const TArray<UObject*>& InObjects, FMenuBuilder& MenuBuilder ) override;
 	virtual void OpenAssetEditor( const TArray<UObject*>& InObjects, TSharedPtr<class IToolkitHost> EditWithinLevelEditor = TSharedPtr<IToolkitHost>() ) override;
 	virtual uint32 GetCategories() override { return EAssetTypeCategories::Animation; }
+	virtual void PerformAssetDiff(UObject* Asset1, UObject* Asset2, const struct FRevisionInfo& OldRevision, const struct FRevisionInfo& NewRevision) const override;
 	virtual class UThumbnailInfo* GetThumbnailInfo(UObject* Asset) const override;
 	virtual TSharedPtr<SWidget> GetThumbnailOverlay(const FAssetData& AssetData) const override;
 
diff --git a/Engine/Source/Editor/UMGEditor/Private/AssetTypeActions_WidgetBlueprint.cpp b/Engine/Source/Editor/UMGEditor/Private/AssetTypeActions_WidgetBlueprint.cpp
--- a/Engine/Source/Editor/UMGEditor/Private/AssetTypeActions_WidgetBlueprint.cpp
+++ b/Engine/Source/Editor/UMGEditor/Private/AssetTypeActions_WidgetBlueprint.cpp
@@ -57,4 +57,49 @@ FText FAssetTypeActions_WidgetBlueprint::GetAssetDescription( const FAssetData&
 	return FText::GetEmpty();
 }
 
+void FAssetTypeActions_WidgetBlueprint::PerformAssetDiff(UObject* Asset1, UObject* Asset2, const struct FRevisionInfo& OldRevision, const struct FRevisionInfo& NewRevision) const
+{
+	if (!GetDefault<UEditorExperimentalSettings>()->bEnableWidgetVisualDiff)
+	{
+		return FAssetTypeActions_Base::PerformAssetDiff(Asset1, Asset2, OldRevision, NewRevision);
+	}
+
+	UBlueprint* OldBlueprint = CastChecked<UBlueprint>(Asset1);
+	UBlueprint* NewBlueprint = CastChecked<UBlueprint>(Asset2);
+
+	// sometimes we're comparing different revisions of one single asset (other 
+	// times we're comparing two completely separate assets altogether)
+	bool bIsSingleAsset = (NewBlueprint->GetName() == OldBlueprint->GetName());
+
+	FText WindowTitle = LOCTEXT("NamelessWidgetBlueprintDiff", "Widget Blueprint Diff");
+	// if we're diff'ing one asset against itself 
+	if (bIsSingleAsset)
+	{
+		// identify the assumed single asset in the window's title
+		WindowTitle = FText::Format(LOCTEXT("WidgetBlueprintDiff", "{0} - Widget Blueprint Diff"), FText::FromString(NewBlueprint->GetName()));
+	}
+
+	const TSharedPtr<SWindow> Window = SNew(SWindow)
+		.Title(WindowTitle)
+		.ClientSize(FVector2D(1000, 800));
+
+	Window->SetContent(SNew(SBlueprintDiff)
+		.BlueprintOld(OldBlueprint)
+		.BlueprintNew(NewBlueprint)
+		.OldRevision(OldRevision)
+		.NewRevision(NewRevision)
+		.ShowAssetNames(!bIsSingleAsset));
+
+	// Make this window a child of the modal window if we've been spawned while one is active.
+	TSharedPtr<SWindow> ActiveModal = FSlateApplication::Get().GetActiveModalWindow();
+	if (ActiveModal.IsValid())
+	{
+		FSlateApplication::Get().AddWindowAsNativeChild(Window.ToSharedRef(), ActiveModal.ToSharedRef());
+	}
+	else
+	{
+		FSlateApplication::Get().AddWindow(Window.ToSharedRef());
+	}
+}
+
 #undef LOCTEXT_NAMESPACE
\ No newline at end of file
diff --git a/Engine/Source/Editor/UMGEditor/Private/AssetTypeActions_WidgetBlueprint.h b/Engine/Source/Editor/UMGEditor/Private/AssetTypeActions_WidgetBlueprint.h
--- a/Engine/Source/Editor/UMGEditor/Private/AssetTypeActions_WidgetBlueprint.h
+++ b/Engine/Source/Editor/UMGEditor/Private/AssetTypeActions_WidgetBlueprint.h
@@ -12,6 +12,7 @@ public:
 	virtual UClass* GetSupportedClass() const override;
 	virtual void OpenAssetEditor(const TArray<UObject*>& InObjects, TSharedPtr<class IToolkitHost> EditWithinLevelEditor = TSharedPtr<IToolkitHost>()) override;
 	virtual uint32 GetCategories() override { return EAssetTypeCategories::UI; }
+	virtual void PerformAssetDiff(UObject* Asset1, UObject* Asset2, const struct FRevisionInfo& OldRevision, const struct FRevisionInfo& NewRevision) const override;
 	virtual FText GetAssetDescription( const FAssetData& AssetData ) const override;
 	// End IAssetTypeActions Implementation
 };
\ No newline at end of file
diff --git a/Engine/Source/Editor/UnrealEd/Classes/Settings/EditorExperimentalSettings.h b/Engine/Source/Editor/UnrealEd/Classes/Settings/EditorExperimentalSettings.h
--- a/Engine/Source/Editor/UnrealEd/Classes/Settings/EditorExperimentalSettings.h
+++ b/Engine/Source/Editor/UnrealEd/Classes/Settings/EditorExperimentalSettings.h
@@ -65,6 +65,14 @@ public:
 	UPROPERTY(EditAnywhere, config, Category=Blueprints, meta=(DisplayName="Blueprint Break on Exceptions"))
 	bool bBreakOnExceptions;
 
+	/** Enables the visual diff tool for widget blueprints. WARNING: changes to the widget hierarchy will not be detected */
+	UPROPERTY(EditAnywhere, config, Category=Blueprints, meta=(DisplayName="Use the Diff Tool for Widget Blueprints"))
+	bool bEnableWidgetVisualDiff;
+
+	/** Enables the visual diff tool for anim blueprints. WARNING: changes to the Target Skeleton and Groups will not be detected */
+	UPROPERTY(EditAnywhere, config, Category = Blueprints, meta = (DisplayName = "Use the Diff Tool for Animation Blueprints"))
+	bool bEnableAnimVisualDiff;
+
 	/** Should arrows indicating data/execution flow be drawn halfway along wires? */
 	UPROPERTY(/*EditAnywhere - deprecated (moved into UBlueprintEditorSettings), */config/*, Category=Blueprints, meta=(DisplayName="Draw midpoint arrows in Blueprints")*/)
 	bool bDrawMidpointArrowsInBlueprints;
