FAssetRenameManager::FindCDOReferencedAssets now ignores deprecated and dead classes

UE-19761 - EDITOR CRASH: Creating, reverting, then moving a same named file caused a crash.

diff --git a/Engine/Source/Developer/AssetTools/Private/AssetRenameManager.cpp b/Engine/Source/Developer/AssetTools/Private/AssetRenameManager.cpp
--- a/Engine/Source/Developer/AssetTools/Private/AssetRenameManager.cpp
+++ b/Engine/Source/Developer/AssetTools/Private/AssetRenameManager.cpp
@@ -9,6 +9,7 @@
 #include "FileHelpers.h"
 #include "ObjectTools.h"
 #include "MainFrame.h"
+#include "Kismet2/KismetEditorUtilities.h"
 
 #define LOCTEXT_NAMESPACE "AssetRenameManager"
 
@@ -249,6 +250,12 @@ TArray<TWeakObjectPtr<UObject>> FAssetRenameManager::FindCDOReferencedAssets(con
 			continue;
 		}
 
+		// Ignore deprecated and temporary trash classes.
+		if (Cls->HasAnyClassFlags(CLASS_Deprecated | CLASS_NewerVersionExists) || FKismetEditorUtilities::IsClassABlueprintSkeleton(Cls))
+		{
+			continue;
+		}
+
 		for (TFieldIterator<UObjectProperty> PropertyIt(Cls); PropertyIt; ++PropertyIt)
 		{
 			const UObject* Object = PropertyIt->GetPropertyValue(PropertyIt->ContainerPtrToValuePtr<UObject>(CDO));
