Copying //UE4/Dev-Sequencer to //UE4/Main (Source: //UE4/Dev-Sequencer @ 2945541)

==========================
MAJOR FEATURES + CHANGES
==========================

Change 2837601 on 2016/01/20 by [REDACTED_PERSON]

	Sequencer: Exposed bool, byte, event, fade, float, slomo, and visibility section keys in context menu

Change 2840895 on 2016/01/23 by [REDACTED_PERSON]

	Sequencer: Moved key proxy handling from section into key area; added support for grouped keys; exposed color properties in context menu.

Change 2901092 on 2016/03/09 by [REDACTED_PERSON]

	Sequencer: Fixed ensure() and crash when exiting PIE & then playing back in editor

	Made sure that instances that are playing that have active montages can re-connect to a playing montage if it is already active. We dont stop/start preview/non-
preview playback when switching modes right now so I've just made the system robust to the switch when playing back animation.
	Ensured that instances are kept up to date by hooking into EndPlayMapDelegate. This fires slightly later than OnEndPIE, so it allows us to pick up when instances
have been deleted from the world. OnEndPIE is fired before the worlds are shutdown so instances can still be bound to valid (but nearly dead) actors at this point.

Change 2901159 on 2016/03/09 by [REDACTED_PERSON]

	Allow EXR frames from SaveHighResScreenshot to be uncompressed (r.SaveUncompressedEXRFrames).

Change 2903123 on 2016/03/10 by [REDACTED_PERSON]

	Sequencer: Add snapping for the in/out range.

Change 2903126 on 2016/03/10 by [REDACTED_PERSON]

	Sequencer: Fix split and trim for cinematic shot sections so that the start offset value is set.

Change 2903132 on 2016/03/10 by [REDACTED_PERSON]

	Sequencer: Create camera here and set it as the current camera cut.

Change 2903138 on 2016/03/10 by [REDACTED_PERSON]

	Sequencer: Improve mechanism for adding objects to Sequencer. Hold down shift while dragging from content browser to add as a spawnable. Hold down ctrl while
dragging from content browser to add as a possessable.

Change 2903143 on 2016/03/10 by [REDACTED_PERSON]

	Sequencer: Add 4k to movie capture dialog.

Change 2903157 on 2016/03/10 by [REDACTED_PERSON]

	Sequencer: Pop out of locked camera when popping out of the sequence.

Change 2908097 on 2016/03/14 by [REDACTED_PERSON]

	Sequencer: Added default viewport type onto the viewport types menu
	  - The activation command for a given viewport type no longer toggles it if it's already active, it will just do nothing.
	  - Switching back to the default viewport type can be achieved with SHIFT+D

Change 2911869 on 2016/03/16 by [REDACTED_PERSON]

	Editor: Caching device proxy manager, so the UI doesn't constantly poll for the TargetDeviceServices module

Change 2917943 on 2016/03/22 by [REDACTED_PERSON]

	Sequence Recording: Added "Record New Sequence From Current Player" to sub-sequence track menu

	This option is only available in PIE. When selected it primes a new seciton for recording against the pawn that is currently being controlled by the player.

Change 2917946 on 2016/03/22 by [REDACTED_PERSON]

	Sequencer: Automatically add and attach a cine camera when dropping a crane or rail.

Change 2917954 on 2016/03/22 by [REDACTED_PERSON]

	Fix anim dynamics going crazy on time skips

	Added new API to FAnimNode_Base: NeedsDynamicReset and ResetDynamics.
	This allows nodes that subscribe to this interface to be reset on teleport/time skips.
	Call through to ResetDynamics on zero timestep (ie. skips) in the preview path in Sequencer.

Change 2917961 on 2016/03/22 by [REDACTED_PERSON]

	Sequencer - Sequencer - Update the color track code to match the behavior in matinee, also remove empty light color tracks from the sequence, and add a missing
particle parameter track.

Change 2917984 on 2016/03/22 by [REDACTED_PERSON]

	Sequencer: Fix crash on deleting object binding nodes multiple times.

Change 2917986 on 2016/03/22 by [REDACTED_PERSON]

	Sequencer: Fix if adding multiple camera cuts at the same time with the same duration as an existing camera cut. The camera is replaced in this case.

Change 2917994 on 2016/03/22 by [REDACTED_PERSON]

	Support for scene depth picker. Used for focus depth sampling in cine cams.

Change 2918003 on 2016/03/22 by [REDACTED_PERSON]

	Sequencer - Fix selection issues related to keying an undo by moving the selection clear on rebuild to after the tree selection state has been cached by path.

Change 2920371 on 2016/03/23 by [REDACTED_PERSON]

	UnrealEd: Added a file import path that doesn't load the entire file into memory prior to importing

	UFactory::StaticImport will now call FactoryCreateFile for both text and binary files. The default implementation will load the file into a buffer/string and call
FactoryCreateBinary/FactoryCreateText to preserve legacy behavior. New factories may override FactoryCreateFile to perform their own file processing instead.

Change 2923359 on 2016/03/25 by [REDACTED_PERSON]

	UnrealEd: Refactored out StaticImportObject code into ImportObject

Change 2924887 on 2016/03/28 by [REDACTED_PERSON]

	Sequencer: Camera Anim and Camera Shake tracks now support PostProcess changes in the CameraAnim

Change 2927283 on 2016/03/30 by [REDACTED_PERSON]

	Sequencer: Fix crash resetting default in audio track in sequencer. Fix audio track rename when adding another audio.

Change 2928290 on 2016/03/30 by [REDACTED_PERSON]

	ContentBrowser: Disabling 'Show in Explorer' context menu option for newly created/unsaved asset(s)

Change 2928480 on 2016/03/30 by [REDACTED_PERSON]

	AssetTools: Added built-in type category for Media assets

Change 2928498 on 2016/03/30 by [REDACTED_PERSON]

	AssetTools: Sorting asset type categories alphabetically

Change 2932326 on 2016/04/04 by [REDACTED_PERSON]

	Sequencer - Add "paste from matinee" support for audio tracks, and add volume support to audio sections in sequencer.

Change 2933917 on 2016/04/05 by [REDACTED_PERSON]

	Core: Added microseconds support to FTimespan

Change 2933920 on 2016/04/05 by [REDACTED_PERSON]

	Sockets: Simplified socket timeout assignments

Change 2935434 on 2016/04/06 by [REDACTED_PERSON]

	Sequencer: Deselect possessable before converting it to a spawnable and deleting it. Also, make sure the newly converted spawnables are selected. This fixes a bug
where the transform gizmo remains after adding a spawnable.

	Also, consolidate code for adding possessables to sequencer so that the newly added possessable node will be selected in the widget tree.

Change 2935590 on 2016/04/06 by [REDACTED_PERSON]

	Made crane rig preview mesh at the tip move as expected and appear in a reasonable place.

Change 2936082 on 2016/04/07 by [REDACTED_PERSON]

	Sequencer: Attach should use default attach component of the actor if it exists. This fixes a bug where trying to attach to a camera crane rig attaches to the
proper component.

Change 2936118 on 2016/04/07 by [REDACTED_PERSON]

	Sequencer: Adding crane/rail with shift now adds the crane/rig as a spawnable and the attached cine camera as a spawnable. An attach track is created for the cine
camera spawnable.

Change 2937226 on 2016/04/07 by [REDACTED_PERSON]

	Sequencer: Add invalidation when levels are added or removed.This fixes a bug were if you have a level sequence open and then add a sublevel that the level
sequence operates on, it correctly updates the actors from the new sub level.

Change 2937263 on 2016/04/07 by [REDACTED_PERSON]

	Sequencer - Changed the matinee to sequencer conversion tools so that curve keys set to "clamped auto" import as user tangents since we don't support clamped auto
in sequencer.

Change 2937273 on 2016/04/07 by [REDACTED_PERSON]

	Sequencer: No longer experimental/betal and now enabled by default.

Change 2937694 on 2016/04/08 by [REDACTED_PERSON]

	Curve Editor: Separate out input and output snapping.

Change 2937852 on 2016/04/08 by [REDACTED_PERSON]

	Sequencer: Spawnables no longer use generated classes

	  - Spawnables now store a template actor instance, rather than a generated class
	  - All relevant code has been converted to use this new approach
	  - Spawnable defaults are now harvested directly from any spawned instances when the movie scene is saved, or when the object is de-spawned. We only do this for
spawnables *in the currently active sequence instance*. This approach ensures that instance components and other properties persist.
	    - Currently we don't mark the package as dirty when the defaults are changed. This needs to be addressed.
	  - Some fixes have been made to AActor and Actor Iterators to ensure that actors not contained within a level do not crash.
	  - Spawning optimisations to follow

Change 2937956 on 2016/04/08 by [REDACTED_PERSON]

	Sequencer: Added selection range commands to general toolbar menu

Change 2937981 on 2016/04/08 by [REDACTED_PERSON]

	Sequencer: Refactored GetKeyHandles to take a time range for filtering keys

Change 2938007 on 2016/04/08 by [REDACTED_PERSON]

	Sequencer: Implemented selection range key selection

Change 2938184 on 2016/04/08 by [REDACTED_PERSON]

	Sequencer: Find in Content Browser is now in the top level menubar.

Change 2938665 on 2016/04/08 by [REDACTED_PERSON]

	UMG - Add support for material animation.

Change 2939048 on 2016/04/10 by [REDACTED_PERSON]

	CineCamera: Change current position on rail to normalize position.

Change 2939067 on 2016/04/10 by [REDACTED_PERSON]

	Sequencer: Add paste color from matinee

Change 2939587 on 2016/04/11 by [REDACTED_PERSON]

	Sequencer: Fixes to spawnable rework
	  - Removed commented out code.
	  - Fixed particle system components not disabling auto activate.
	  - Instance Components are now correctly registered on spawn.
	  - Removed ability to set the editable flag on actors in favor of a delegate assigned to SActorDetails to disable property editing on actors spawned from outside
of the currently focused sequence.

Change 2939666 on 2016/04/11 by [REDACTED_PERSON]

	Sequencer: Fixed active sequence ID not being set on creation

Change 2940663 on 2016/04/12 by [REDACTED_PERSON]

	Sequencer: Spawnable object templates now always have the RF_ArchetypeObject flag

Change 2940742 on 2016/04/12 by [REDACTED_PERSON]

	Improvements to sequence recording API

	Sequence recording is now more extensible.
	Moved the majority of recorders into a seperate module.
	Cleaned up recorder API, removed bRecord that was only really used once.
	Added factory class as a modular feature that allows users to register and implement their own recorders.
	Recorders can supply their own settings classes. These are incorporated into the UI via a details customization so they appear to be seamlessly integrated.
	Exposed components to record as an advanced setting. Users can use this to gate the components and actors that are recorded.

Change 2940828 on 2016/04/12 by [REDACTED_PERSON]

	Level Editor: Clear the camera preview if the level viewport is locked to the same camera.

Change 2941090 on 2016/04/12 by [REDACTED_PERSON]

	Sequencer - Fix inconsistencies with the level visibility track.
	+ Add an option to update a track instance when it stops playing due to it's parent movie scene being deactivated as a sub-scene.
	+ Reset level visibility state when the section ends, or when it's parent subsection ends.

Change 2941427 on 2016/04/12 by [REDACTED_PERSON]

	Sequencer: Fixes for recorded dynamic components
	  - Dynamic components are now named uniquely within their owner actor, and are no longer re-bound to new object tracks when detached/re-attached
	  - New dynamic components are created every time a new component is detected on an actor, regardless of whether it was once attached before. This affords
recording externally managed components from a component pool.
	  - Section recorders for components that are no longer attached to the actor are now disabled correctly.

Change 2941760 on 2016/04/12 by [REDACTED_PERSON]

	Editor: Setting up correct material expresions when creating material from normal map texture

Change 2941819 on 2016/04/13 by [REDACTED_PERSON]

	Sequencer: Fire named events in all sublevels.

Change 2942052 on 2016/04/13 by [REDACTED_PERSON]

	Sequencer: Fixed not being able to add dynamic material parameter tracks to spawnables when the object is not spawned

Change 2942223 on 2016/04/13 by [REDACTED_PERSON]

	Editor: Using built-in function to assign sampler format

Change 2942833 on 2016/04/13 by [REDACTED_PERSON]

	Sequencer: More copy matinee helpers.

	- Make some sections infinite.
	- Fix FindPossessableObjectId
	- Expose FindGroupByName

Change 2942944 on 2016/04/13 by [REDACTED_PERSON]

	Sequencer: Fix anim copy so that it doesn't create an animation clip longer than the next clip's start position.

Change 2943313 on 2016/04/14 by [REDACTED_PERSON]

	Sequencer: Initial matinee to level sequence asset converter. Right click on a matinee actor in the level and choose "Convert to Level Sequence"

	Current supports:
	 - movement tracks
	 - generic property tracks
	 - particle tracks
	 - anim control tracks
	 - event tracks
	 - audio tracks
	 - visibility track
	 - director track (fade, slomo, camera cuts)

	Possible todos:
	- Relative scale3D
	- Option to create spawnables instead of possessables
	- Create shots per director track camera cut instead of a direct translation to the camera cut track

Change 2943596 on 2016/04/14 by [REDACTED_PERSON]

	Editor: Fixed placement mode scrollbar visibility not working for custom content

Change 2943651 on 2016/04/14 by [REDACTED_PERSON]

	Sequencer: Invalidate the audio waveform when the source changes.

Change 2943674 on 2016/04/14 by [REDACTED_PERSON]

	Sequencer: Store/restore playback state when rebuilding.

Change 2943993 on 2016/04/14 by [REDACTED_PERSON]

	Sequencer: Convert slot name from matinee to level sequence.

Change 2944156 on 2016/04/14 by [REDACTED_PERSON]

	Sequencer - Fix undo when moving items into folders, and when deleting folder.

Change 2944227 on 2016/04/14 by [REDACTED_PERSON]

	Sequencer: Show curves in the curve editor if one of the parent nodes is selected. For example, if Location is selected, show Location.X, Location.Y, and
Location.Z

Change 2945057 on 2016/04/15 by [REDACTED_PERSON]

	Editor: Fixed placement mode scrollbar visibility not working for custom content

diff --git a/Engine/Config/BaseEngine.ini b/Engine/Config/BaseEngine.ini
--- a/Engine/Config/BaseEngine.ini
+++ b/Engine/Config/BaseEngine.ini
@@ -2594,19 +2594,20 @@ bBuildForD3D12=false
 +Settings=(ScreenSize=.1,OptimizationSettings=(NumOfTrianglesPercentage=.06))
 
 [/Script/CinematicCamera.CineCameraComponent]
++FilmbackPresets=(Name="16:9 Film",FilmbackSettings=(SensorWidth=24.892,SensorHeight=14.002))
++FilmbackPresets=(Name="16:9 Digital Film",FilmbackSettings=(SensorWidth=23.66,SensorHeight=13.3))
++FilmbackPresets=(Name="16:9 DSLR",FilmbackSettings=(SensorWidth=36,SensorHeight=20.25))
 +FilmbackPresets=(Name="Super 8mm",FilmbackSettings=(SensorWidth=5.79,SensorHeight=4.01))
 +FilmbackPresets=(Name="Super 16mm",FilmbackSettings=(SensorWidth=12.52,SensorHeight=7.58))
 +FilmbackPresets=(Name="Super 35mm",FilmbackSettings=(SensorWidth=24.89,SensorHeight=18.66))
 +FilmbackPresets=(Name="35mm Academy",FilmbackSettings=(SensorWidth=21.946,SensorHeight=16.002))
 +FilmbackPresets=(Name="35mm Full Aperture",FilmbackSettings=(SensorWidth=24.892,SensorHeight=18.9121))
 +FilmbackPresets=(Name="35mm VistaVision",FilmbackSettings=(SensorWidth=25.146,SensorHeight=37.719))
-;+FilmbackPresets=(Name="65mm 15 perf",FilmbackSettings=(SensorWidth=52.48,SensorHeight=23.01))
 +FilmbackPresets=(Name="IMAX 70mm",FilmbackSettings=(SensorWidth=70.41,SensorHeight=56.63))
 +FilmbackPresets=(Name="APS-C (Canon)",FilmbackSettings=(SensorWidth=22.2,SensorHeight=14.8))
 +FilmbackPresets=(Name="Full Frame DSLR",FilmbackSettings=(SensorWidth=36,SensorHeight=24))
 +FilmbackPresets=(Name="Micro Four Thirds",FilmbackSettings=(SensorWidth=17.3,SensorHeight=13))
-; Sony F65
-; Arri Alexa
+DefaultFilmbackPresetName="16:9 DSLR"
 
 +LensPresets=(Name="12mm Prime f/2.8",LensSettings=(MinFocalLength=12,MaxFocalLength=12,MinFStop=2.8,MaxFStop=22))
 +LensPresets=(Name="30mm Prime f/1.4",LensSettings=(MinFocalLength=30,MaxFocalLength=30,MinFStop=1.4,MaxFStop=22))
@@ -2616,4 +2617,7 @@ bBuildForD3D12=false
 +LensPresets=(Name="200mm Prime f/2",LensSettings=(MinFocalLength=200,MaxFocalLength=200,MinFStop=2,MaxFStop=22))
 +LensPresets=(Name="24-70mm Zoom f/2.8",LensSettings=(MinFocalLength=24,MaxFocalLength=70,MinFStop=2.8,MaxFStop=22))
 +LensPresets=(Name="70-200mm Zoom f/2.8",LensSettings=(MinFocalLength=70,MaxFocalLength=200,MinFStop=2.8,MaxFStop=22))
-
++LensPresets=(Name="Universal Zoom",LensSettings=(MinFocalLength=4,MaxFocalLength=1000,MinFStop=1.2,MaxFStop=22))
+DefaultLensPresetName="Universal Zoom"
+DefaultLensFocalLength=35
+DefaultLensFStop=2.8
diff --git a/Engine/Plugins/2D/Paper2D/Source/PaperSpriteSheetImporter/Private/PaperSpriteSheetReimportFactory.cpp b/Engine/Plugins/2D/Paper2D/Source/PaperSpriteSheetImporter/Private/PaperSpriteSheetReimportFactory.cpp
--- a/Engine/Plugins/2D/Paper2D/Source/PaperSpriteSheetImporter/Private/PaperSpriteSheetReimportFactory.cpp
+++ b/Engine/Plugins/2D/Paper2D/Source/PaperSpriteSheetImporter/Private/PaperSpriteSheetReimportFactory.cpp
@@ -61,7 +61,9 @@ EReimportResult::Type UPaperSpriteSheetReimportFactory::Reimport(UObject* Obj)
 
 	// Run the import again
 	EReimportResult::Type Result = EReimportResult::Failed;
-	if (UFactory::StaticImportObject(SpriteSheet->GetClass(), SpriteSheet->GetOuter(), *SpriteSheet->GetName(), RF_Public | RF_Standalone, *Filename, nullptr, this))
+	bool OutCanceled = false;
+
+	if (ImportObject(SpriteSheet->GetClass(), SpriteSheet->GetOuter(), *SpriteSheet->GetName(), RF_Public | RF_Standalone, Filename, nullptr, OutCanceled) != nullptr)
 	{
 		UE_LOG(LogPaperSpriteSheetImporter, Log, TEXT("Imported successfully"));
 
@@ -80,7 +82,15 @@ EReimportResult::Type UPaperSpriteSheetReimportFactory::Reimport(UObject* Obj)
 	}
 	else
 	{
-		UE_LOG(LogPaperSpriteSheetImporter, Warning, TEXT("-- import failed"));
+		if (OutCanceled)
+		{
+			UE_LOG(LogPaperSpriteSheetImporter, Warning, TEXT("-- import canceled"));
+		}
+		else
+		{
+			UE_LOG(LogPaperSpriteSheetImporter, Warning, TEXT("-- import failed"));
+		}
+
 		Result = EReimportResult::Failed;
 	}
 
diff --git a/Engine/Plugins/Editor/SpeedTreeImporter/Source/SpeedTreeImporter/Private/ReimportSpeedTreeFactory.cpp b/Engine/Plugins/Editor/SpeedTreeImporter/Source/SpeedTreeImporter/Private/ReimportSpeedTreeFactory.cpp
--- a/Engine/Plugins/Editor/SpeedTreeImporter/Source/SpeedTreeImporter/Private/ReimportSpeedTreeFactory.cpp
+++ b/Engine/Plugins/Editor/SpeedTreeImporter/Source/SpeedTreeImporter/Private/ReimportSpeedTreeFactory.cpp
@@ -78,22 +78,27 @@ EReimportResult::Type UReimportSpeedTreeFactory::Reimport(UObject* Obj)
 
 	UE_LOG(LogEditorFactories, Log, TEXT("Performing atomic reimport of [%s]"), *Filename);
 
-	if (UFactory::StaticImportObject(Mesh->GetClass(), Mesh->GetOuter(), *Mesh->GetName(), RF_Public | RF_Standalone, *Filename, NULL, this))
+	bool OutCanceled = false;
+
+	if (ImportObject(Mesh->GetClass(), Mesh->GetOuter(), *Mesh->GetName(), RF_Public | RF_Standalone, Filename, nullptr, OutCanceled) != nullptr)
 	{
 		UE_LOG(LogEditorFactories, Log, TEXT("-- imported successfully"));
 
 		Mesh->AssetImportData->Update(Filename);
-
-		// Mark the package dirty after the successful import
 		Mesh->MarkPackageDirty();
 
 		return EReimportResult::Succeeded;
 	}
+	else if (OutCanceled)
+	{
+		UE_LOG(LogEditorFactories, Warning, TEXT("-- import canceled"));
+	}
 	else
 	{
 		UE_LOG(LogEditorFactories, Warning, TEXT("-- import failed"));
 	}
 #endif // #if WITH_SPEEDTREE
+
 	return EReimportResult::Failed;
 }
 
diff --git a/Engine/Plugins/Messaging/UdpMessaging/Source/UdpMessaging/Private/UdpMessagingModule.cpp b/Engine/Plugins/Messaging/UdpMessaging/Source/UdpMessaging/Private/UdpMessagingModule.cpp
--- a/Engine/Plugins/Messaging/UdpMessaging/Source/UdpMessaging/Private/UdpMessagingModule.cpp
+++ b/Engine/Plugins/Messaging/UdpMessaging/Source/UdpMessaging/Private/UdpMessagingModule.cpp
@@ -36,6 +36,9 @@ public:
 		{
 			UUdpMessagingSettings* Settings = GetMutableDefault<UUdpMessagingSettings>();
 
+			// general information
+			Ar.Logf(TEXT("Protocol Version: %i"), UDP_MESSAGING_TRANSPORT_PROTOCOL_VERSION);
+
 			// bridge status
 			if (MessageBridge.IsValid())
 			{
@@ -54,22 +57,38 @@ public:
 			}
 
 			// bridge settings
-			Ar.Logf(TEXT("  Unicast Endpoint: %s"), *Settings->UnicastEndpoint);
-			Ar.Logf(TEXT("  Multicast Endpoint: %s"), *Settings->MulticastEndpoint);
-			Ar.Logf(TEXT("  Multicast TTL: %i"), Settings->MulticastTimeToLive);
+			if (Settings->UnicastEndpoint.IsEmpty())
+			{
+				Ar.Logf(TEXT("    Unicast Endpoint: %s (default)"), *FIPv4Endpoint::Any.ToString());
+			}
+			else
+			{
+				Ar.Logf(TEXT("    Unicast Endpoint: %s"), *Settings->UnicastEndpoint);
+			}
+
+			if (Settings->MulticastEndpoint.IsEmpty())
+			{
+				Ar.Logf(TEXT("    Multicast Endpoint: %s (default)"), *UDP_MESSAGING_DEFAULT_MULTICAST_ENDPOINT.ToString());
+			}
+			else
+			{
+				Ar.Logf(TEXT("    Multicast Endpoint: %s"), *Settings->MulticastEndpoint);
+			}
+
+			Ar.Logf(TEXT("    Multicast TTL: %i"), Settings->MulticastTimeToLive);
 
 			if (Settings->StaticEndpoints.Num() > 0)
 			{
-				Ar.Log(TEXT("  Static Endpoints:"));
+				Ar.Log(TEXT("    Static Endpoints:"));
 
 				for (const auto& StaticEndpoint : Settings->StaticEndpoints)
 				{
-					Ar.Logf(TEXT("  > %s"), *StaticEndpoint);
+					Ar.Logf(TEXT("        %s"), *StaticEndpoint);
 				}
 			}
 			else
 			{
-				Ar.Log(TEXT("  Static Endpoints: None"));
+				Ar.Log(TEXT("    Static Endpoints: None"));
 			}
 
 #if PLATFORM_DESKTOP
@@ -90,19 +109,44 @@ public:
 				Ar.Log(TEXT("Message Tunnel: Not initialized."));
 			}
 
-			Ar.Logf(TEXT("  Unicast Endpoint: %s"), *Settings->TunnelUnicastEndpoint);
-			Ar.Logf(TEXT("  Multicast Endpoint: %s"), *Settings->TunnelMulticastEndpoint);
-			Ar.Log(TEXT("  Remote Endpoints:"));
+			// tunnel settings
+			if (Settings->TunnelUnicastEndpoint.IsEmpty())
+			{
+				Ar.Logf(TEXT("    Unicast Endpoint: %s (default)"), *FIPv4Endpoint::Any.ToString());
+			}
+			else
+			{
+				Ar.Logf(TEXT("    Unicast Endpoint: %s"), *Settings->TunnelUnicastEndpoint);
+			}
+
+			if (Settings->TunnelMulticastEndpoint.IsEmpty())
+			{
+				Ar.Logf(TEXT("    Multicast Endpoint: %s (default)"), *UDP_MESSAGING_DEFAULT_MULTICAST_ENDPOINT.ToString());
+			}
+			else
+			{
+				Ar.Logf(TEXT("    Multicast Endpoint: %s"), *Settings->TunnelMulticastEndpoint);
+			}
+
+			if (Settings->RemoteTunnelEndpoints.Num() > 0)
+			{
+				Ar.Log(TEXT("    Remote Endpoints:"));
 
-			for (const auto& RemoteEndpoint : Settings->RemoteTunnelEndpoints)
+				for (const auto& RemoteEndpoint : Settings->RemoteTunnelEndpoints)
+				{
+					Ar.Logf(TEXT("        %s"), *RemoteEndpoint);
+				}
+			}
+			else
 			{
-				Ar.Logf(TEXT("  > %s"), *RemoteEndpoint);
+				Ar.Log(TEXT("    Remote Endpoints: None"));
 			}
 
+			// tunnel performance
 			if (MessageTunnel.IsValid())
 			{
-				Ar.Logf(TEXT("  Total Bytes In: %i"), MessageTunnel->GetTotalInboundBytes());
-				Ar.Logf(TEXT("  Total Bytes Out: %i"), MessageTunnel->GetTotalOutboundBytes());
+				Ar.Logf(TEXT("    Total Bytes In: %i"), MessageTunnel->GetTotalInboundBytes());
+				Ar.Logf(TEXT("    Total Bytes Out: %i"), MessageTunnel->GetTotalOutboundBytes());
 
 				TArray<TSharedPtr<IUdpMessageTunnelConnection>> Connections;
 			
diff --git a/Engine/Plugins/MovieScene/LevelSequenceEditor/LevelSequenceEditor.uplugin b/Engine/Plugins/MovieScene/LevelSequenceEditor/LevelSequenceEditor.uplugin
--- a/Engine/Plugins/MovieScene/LevelSequenceEditor/LevelSequenceEditor.uplugin
+++ b/Engine/Plugins/MovieScene/LevelSequenceEditor/LevelSequenceEditor.uplugin
@@ -10,9 +10,9 @@
 	"DocsURL" : "",
 	"MarketplaceURL" : "",
 	"SupportURL" : "",
-	"EnabledByDefault" : false,
+	"EnabledByDefault" : true,
 	"CanContainContent" : false,
-	"IsBetaVersion" : true,
+	"IsBetaVersion" : false,
 	"Installed" : false,
 	"Modules" :
 	[
diff --git a/Engine/Plugins/MovieScene/LevelSequenceEditor/Source/LevelSequenceEditor/Private/CinematicViewport/CinematicViewportLayoutEntity.h b/Engine/Plugins/MovieScene/LevelSequenceEditor/Source/LevelSequenceEditor/Private/CinematicViewport/CinematicViewportLayoutEntity.h
--- a/Engine/Plugins/MovieScene/LevelSequenceEditor/Source/LevelSequenceEditor/Private/CinematicViewport/CinematicViewportLayoutEntity.h
+++ b/Engine/Plugins/MovieScene/LevelSequenceEditor/Source/LevelSequenceEditor/Private/CinematicViewport/CinematicViewportLayoutEntity.h
@@ -4,6 +4,7 @@
 
 #include "LevelViewportLayout.h"
 #include "SCinematicLevelViewport.h"
+#include "UnrealEd.h"
 
 class FCinematicViewportLayoutEntity : public IViewportLayoutEntity
 {
@@ -27,11 +28,17 @@ public:
 	}
 
 	FLevelEditorViewportClient& GetLevelViewportClient() const { return Widget->GetLevelViewport()->GetLevelViewportClient(); }
-	bool IsPlayInEditorViewportActive() const { return false; }
-	void RegisterGameViewportIfPIE(){}
-	void SetKeyboardFocus(){ FSlateApplication::Get().SetKeyboardFocus(Widget); }
-	void OnLayoutDestroyed(){}
-	void SaveConfig(const FString&){}
+	bool IsPlayInEditorViewportActive() const { return Widget->GetLevelViewport()->IsPlayInEditorViewportActive(); }
+	void RegisterGameViewportIfPIE(){ return Widget->GetLevelViewport()->RegisterGameViewportIfPIE(); }
+	void SetKeyboardFocus(){ FSlateApplication::Get().SetKeyboardFocus(Widget->GetLevelViewport()); }
+	void OnLayoutDestroyed()
+	{
+		if (IsPlayInEditorViewportActive() || GetLevelViewportClient().IsSimulateInEditorViewport() )
+		{
+			GUnrealEd->EndPlayMap();
+		}
+	}
+	void SaveConfig(const FString& ConfigString){ return Widget->GetLevelViewport()->SaveConfig(ConfigString); }
 
 protected:
 
diff --git a/Engine/Plugins/MovieScene/LevelSequenceEditor/Source/LevelSequenceEditor/Private/CinematicViewport/SCinematicLevelViewport.cpp b/Engine/Plugins/MovieScene/LevelSequenceEditor/Source/LevelSequenceEditor/Private/CinematicViewport/SCinematicLevelViewport.cpp
--- a/Engine/Plugins/MovieScene/LevelSequenceEditor/Source/LevelSequenceEditor/Private/CinematicViewport/SCinematicLevelViewport.cpp
+++ b/Engine/Plugins/MovieScene/LevelSequenceEditor/Source/LevelSequenceEditor/Private/CinematicViewport/SCinematicLevelViewport.cpp
@@ -573,8 +573,6 @@ void SCinematicLevelViewport::Tick(const FGeometry& AllottedGeometry, const doub
 		return;
 	}
 
-	ViewportClient->bDisableInput = Sequencer->GetPlaybackStatus() == EMovieScenePlayerStatus::Playing;
-
 	UMovieSceneSequence* Sequence = Sequencer->GetFocusedMovieSceneSequence();
 	if (!Sequence)
 	{
diff --git a/Engine/Plugins/MovieScene/LevelSequenceEditor/Source/LevelSequenceEditor/Private/LevelSequenceEditorModule.cpp b/Engine/Plugins/MovieScene/LevelSequenceEditor/Source/LevelSequenceEditor/Private/LevelSequenceEditorModule.cpp
--- a/Engine/Plugins/MovieScene/LevelSequenceEditor/Source/LevelSequenceEditor/Private/LevelSequenceEditorModule.cpp
+++ b/Engine/Plugins/MovieScene/LevelSequenceEditor/Source/LevelSequenceEditor/Private/LevelSequenceEditorModule.cpp
@@ -86,8 +86,7 @@ protected:
 	{
 		FLevelEditorModule& LevelEditorModule = FModuleManager::LoadModuleChecked<FLevelEditorModule>("LevelEditor");
 
-		FViewportTypeDefinition CinematicViewportType = FViewportTypeDefinition::FromType<FCinematicViewportLayoutEntity>();
-		CinematicViewportType.ToggleCommand = FLevelSequenceEditorCommands::Get().ToggleCinematicViewportCommand;
+		FViewportTypeDefinition CinematicViewportType = FViewportTypeDefinition::FromType<FCinematicViewportLayoutEntity>(FLevelSequenceEditorCommands::Get().ToggleCinematicViewportCommand);
 		LevelEditorModule.RegisterViewportType("Cinematic", CinematicViewportType);
 	}
 
diff --git a/Engine/Plugins/MovieScene/LevelSequenceEditor/Source/LevelSequenceEditor/Private/LevelSequenceEditorToolkit.cpp b/Engine/Plugins/MovieScene/LevelSequenceEditor/Source/LevelSequenceEditor/Private/LevelSequenceEditorToolkit.cpp
--- a/Engine/Plugins/MovieScene/LevelSequenceEditor/Source/LevelSequenceEditor/Private/LevelSequenceEditorToolkit.cpp
+++ b/Engine/Plugins/MovieScene/LevelSequenceEditor/Source/LevelSequenceEditor/Private/LevelSequenceEditorToolkit.cpp
@@ -18,12 +18,12 @@
 #include "SceneOutlinerModule.h"
 #include "SceneOutlinerPublicTypes.h"
 #include "SDockTab.h"
+#include "SequencerSpawnRegister.h"
 
 // @todo sequencer: hack: setting defaults for transform tracks
 #include "MovieScene3DTransformSection.h"
 #include "MovieScene3DTransformTrack.h"
 
-#include "SequencerSpawnRegister.h"
 
 #define LOCTEXT_NAMESPACE "LevelSequenceEditor"
 
@@ -225,87 +225,24 @@ void FLevelSequenceEditorToolkit::UnregisterTabSpawners(const TSharedRef<class F
 }
 
 
-/* FLevelSequenceEditorToolkit callbacks
+/* FLevelSequenceEditorToolkit implementation
  *****************************************************************************/
 
-void FLevelSequenceEditorToolkit::HandleAddComponentActionExecute(UActorComponent* Component)
-{
-	Sequencer->GetHandleToObject(Component);
-}
-
-
-void FLevelSequenceEditorToolkit::HandleAddComponentMaterialActionExecute(UPrimitiveComponent* Component, int32 MaterialIndex)
-{
-	FGuid ObjectHandle = Sequencer->GetHandleToObject(Component);
-	for (const FMovieSceneBinding& Binding : Sequencer->GetFocusedMovieSceneSequence()->GetMovieScene()->GetBindings())
-	{
-		if (Binding.GetObjectGuid() == ObjectHandle)
-		{
-			bool bHasMaterialTrack = false;
-			for (UMovieSceneTrack* Track : Binding.GetTracks())
-			{
-				UMovieSceneComponentMaterialTrack* MaterialTrack = Cast<UMovieSceneComponentMaterialTrack>(Track);
-				if (MaterialTrack != nullptr && MaterialTrack->GetMaterialIndex() == MaterialIndex)
-				{
-					bHasMaterialTrack = true;
-					break;
-				}
-			}
-			if (bHasMaterialTrack == false)
-			{
-				const FScopedTransaction Transaction(LOCTEXT("AddComponentMaterialTrack", "Add component material track"));
-
-				UMovieScene* FocusedMovieScene = Sequencer->GetFocusedMovieSceneSequence()->GetMovieScene();
-				{
-					FocusedMovieScene->Modify();
-				}
-
-				UMovieSceneComponentMaterialTrack* MaterialTrack = Cast<UMovieSceneComponentMaterialTrack>(FocusedMovieScene->AddTrack<UMovieSceneComponentMaterialTrack>(ObjectHandle));
-				{
-					MaterialTrack->Modify();
-					MaterialTrack->SetMaterialIndex(MaterialIndex);
-				}
-				
-				Sequencer->NotifyMovieSceneDataChanged();
-			}
-		}
-	}
-}
-
 void FLevelSequenceEditorToolkit::AddActorsToSequencer(AActor*const* InActors, int32 NumActors)
 {
-	const FScopedTransaction Transaction(LOCTEXT("UndoPossessingObject", "Possess Object(s) with Sequencer"));
-
-	UMovieSceneSequence* FocussedSequence = Sequencer->GetFocusedMovieSceneSequence();
-	UMovieScene* FocussedMovieScene = FocussedSequence->GetMovieScene();
+	TArray<TWeakObjectPtr<AActor>> Actors;
 
-	FocussedSequence->Modify();
-
-	GEditor->SelectNone(true, true);
 	while (NumActors--)
 	{
 		AActor* ThisActor = *InActors;
+		Actors.Add(ThisActor);
 
-		if (!Sequencer->GetFocusedMovieSceneSequenceInstance()->FindObjectId(*ThisActor).IsValid())
-		{
-			FGuid Binding = Sequencer->CreateBinding(*ThisActor, ThisActor->GetActorLabel());
-
-			Sequencer->UpdateRuntimeInstances();
-
-			AddDefaultTracksForActor(*ThisActor, Binding);
-		}
-
-		GEditor->SelectActor(ThisActor, true, true);
 		InActors++;
 	}
 
-	Sequencer->NotifyMovieSceneDataChanged();
+	Sequencer->AddActors(Actors);
 }
 
-void FLevelSequenceEditorToolkit::HandleActorAddedToSequencer(AActor* Actor, const FGuid Binding)
-{
-	AddDefaultTracksForActor(*Actor, Binding);
-}
 
 void FLevelSequenceEditorToolkit::AddDefaultTracksForActor(AActor& Actor, const FGuid Binding)
 {
@@ -462,6 +399,7 @@ void FLevelSequenceEditorToolkit::AddDefaultTracksForActor(AActor& Actor, const
 	}
 }
 
+
 void FLevelSequenceEditorToolkit::AddPosessActorMenuExtensions(FMenuBuilder& MenuBuilder)
 {
 	auto IsActorValidForPossession = [=](const AActor* Actor){
@@ -470,6 +408,7 @@ void FLevelSequenceEditorToolkit::AddPosessActorMenuExtensions(FMenuBuilder& Men
 		{
 			return !Sequencer->GetHandleToObject((UObject*)Actor, bCreateHandleIfMissing).IsValid();
 		}
+
 		return true;
 	};
 
@@ -480,6 +419,7 @@ void FLevelSequenceEditorToolkit::AddPosessActorMenuExtensions(FMenuBuilder& Men
 
 	FText SelectedLabel;
 	FSlateIcon ActorIcon(FEditorStyle::GetStyleSetName(), FClassIconFinder::FindIconNameForClass(AActor::StaticClass()));
+
 	if (ActorsValidForPossession.Num() == 1)
 	{
 		SelectedLabel = FText::Format(LOCTEXT("AddSpecificActor", "Add '{0}'"), FText::FromString(ActorsValidForPossession[0]->GetActorLabel()));
@@ -537,11 +477,47 @@ void FLevelSequenceEditorToolkit::AddPosessActorMenuExtensions(FMenuBuilder& Men
 				})
 			)
 		];
+
 	MenuBuilder.AddWidget(MiniSceneOutliner, FText::GetEmpty(), true);
 	MenuBuilder.EndSection();
 }
 
 
+/* FLevelSequenceEditorToolkit callbacks
+ *****************************************************************************/
+
+void FLevelSequenceEditorToolkit::HandleAddComponentActionExecute(UActorComponent* Component)
+{
+	Sequencer->GetHandleToObject(Component);
+}
+
+
+void FLevelSequenceEditorToolkit::HandleAddComponentMaterialActionExecute(UPrimitiveComponent* Component, int32 MaterialIndex)
+{
+	FGuid ObjectHandle = Sequencer->GetHandleToObject(Component);
+	UMovieScene* FocusedMovieScene = Sequencer->GetFocusedMovieSceneSequence()->GetMovieScene();
+	FName IndexName( *FString::FromInt(MaterialIndex) );
+	if ( FocusedMovieScene->FindTrack( UMovieSceneComponentMaterialTrack::StaticClass(), ObjectHandle, IndexName ) == nullptr )
+	{
+		const FScopedTransaction Transaction( LOCTEXT( "AddComponentMaterialTrack", "Add component material track" ) );
+
+		FocusedMovieScene->Modify();
+
+		UMovieSceneComponentMaterialTrack* MaterialTrack = Cast<UMovieSceneComponentMaterialTrack>( FocusedMovieScene->AddTrack<UMovieSceneComponentMaterialTrack>( ObjectHandle ) );
+		MaterialTrack->Modify();
+		MaterialTrack->SetMaterialIndex( MaterialIndex );
+
+		Sequencer->NotifyMovieSceneDataChanged();
+	}
+}
+
+
+void FLevelSequenceEditorToolkit::HandleActorAddedToSequencer(AActor* Actor, const FGuid Binding)
+{
+	AddDefaultTracksForActor(*Actor, Binding);
+}
+
+
 void FLevelSequenceEditorToolkit::HandleMapChanged(class UWorld* NewWorld, EMapChangeType MapChangeType)
 {
 	Sequencer->NotifyMapChanged(NewWorld, MapChangeType);
@@ -594,10 +570,13 @@ void FLevelSequenceEditorToolkit::HandleTrackMenuExtensionAddTrack(FMenuBuilder&
 		{
 			for (UActorComponent* Component : Actor->GetComponents())
 			{
-				FUIAction AddComponentAction(FExecuteAction::CreateSP(this, &FLevelSequenceEditorToolkit::HandleAddComponentActionExecute, Component));
-				FText AddComponentLabel = FText::FromString(Component->GetName());
-				FText AddComponentToolTip = FText::Format(LOCTEXT("ComponentToolTipFormat", "Add {0} component"), FText::FromString(Component->GetName()));
-				AddTrackMenuBuilder.AddMenuEntry(AddComponentLabel, AddComponentToolTip, FSlateIcon(), AddComponentAction);
+				if (Component)
+				{
+					FUIAction AddComponentAction(FExecuteAction::CreateSP(this, &FLevelSequenceEditorToolkit::HandleAddComponentActionExecute, Component));
+					FText AddComponentLabel = FText::FromString(Component->GetName());
+					FText AddComponentToolTip = FText::Format(LOCTEXT("ComponentToolTipFormat", "Add {0} component"), FText::FromString(Component->GetName()));
+					AddTrackMenuBuilder.AddMenuEntry(AddComponentLabel, AddComponentToolTip, FSlateIcon(), AddComponentAction);
+				}
 			}
 		}
 		AddTrackMenuBuilder.EndSection();
diff --git a/Engine/Plugins/MovieScene/LevelSequenceEditor/Source/LevelSequenceEditor/Private/LevelSequenceEditorToolkit.h b/Engine/Plugins/MovieScene/LevelSequenceEditor/Source/LevelSequenceEditor/Private/LevelSequenceEditorToolkit.h
--- a/Engine/Plugins/MovieScene/LevelSequenceEditor/Source/LevelSequenceEditor/Private/LevelSequenceEditorToolkit.h
+++ b/Engine/Plugins/MovieScene/LevelSequenceEditor/Source/LevelSequenceEditor/Private/LevelSequenceEditorToolkit.h
@@ -5,6 +5,7 @@
 #include "Toolkits/AssetEditorToolkit.h"
 #include "LevelSequence.h"
 
+
 enum class EMapChangeType : uint8;
 class FTabManager;
 class ILevelViewport;
@@ -60,9 +61,30 @@ public:
 	 */
 	void Initialize(const EToolkitMode::Type Mode, const TSharedPtr<IToolkitHost>& InitToolkitHost, ULevelSequence* LevelSequence, bool bEditWithinLevelEditor);
 
+	/**
+	 * Get the sequencer object being edited in this tool kit.
+	 *
+	 * @return Sequencer object.
+	 */
+	TSharedPtr<ISequencer> GetSequencer() const
+	{
+		return Sequencer;
+	}
+
 public:
 
-	// IToolkit interface
+	//~ FAssetEditorToolkit interface
+
+	virtual void AddReferencedObjects(FReferenceCollector& Collector) override
+	{
+		Collector.AddReferencedObject(LevelSequence);
+	}
+
+	virtual bool OnRequestClose() override;
+
+public:
+
+	//~ IToolkit interface
 
 	virtual FText GetBaseToolkitName() const override;
 	virtual FName GetToolkitFName() const override;
@@ -71,14 +93,16 @@ public:
 	virtual void RegisterTabSpawners(const TSharedRef<FTabManager>& TabManager) override;
 	virtual void UnregisterTabSpawners(const TSharedRef<FTabManager>& TabManager) override;
 
-	void UpdateViewports(AActor* ActorToViewThrough) const;
+protected:
 
-	void AddReferencedObjects( FReferenceCollector& Collector ) override
-	{
-		Collector.AddReferencedObject(LevelSequence);
-	}
+	/** Add the specified actors to the sequencer */
+	void AddActorsToSequencer(AActor*const* InActors, int32 NumActors);
 
-	TSharedPtr<ISequencer> GetSequencer() const { return Sequencer; }
+	/** Add default movie scene tracks for the given actor. */
+	void AddDefaultTracksForActor(AActor& Actor, const FGuid Binding);
+
+	/** Menu extension callback for the add menu */
+	void AddPosessActorMenuExtensions(FMenuBuilder& MenuBuilder);
 
 private:
 
@@ -100,28 +124,17 @@ private:
 	/** Callback for the track menu extender. */
 	void HandleTrackMenuExtensionAddTrack(FMenuBuilder& AddTrackMenuBuilder, TArray<UObject*> ContextObjects);
 
-	/** Menu extension callback for the add menu */
-	void AddPosessActorMenuExtensions(FMenuBuilder& MenuBuilder);
-
-	/** Add the specified actors to the sequencer */
-	void AddActorsToSequencer(AActor*const* InActors, int32 NumActors);
-
 	/** Callback for actor added to sequencer. */
 	void HandleActorAddedToSequencer(AActor* Actor, const FGuid Binding);
 
-	/** Add default movie scene tracks for the given actor. */
-	void AddDefaultTracksForActor(AActor& Actor, const FGuid Binding);
-
-private:
-
-	/** Called when this toolkit is to be destroyed */
-	virtual bool OnRequestClose() override;
-
 private:
 
 	/** Level sequence for our edit operation. */
 	ULevelSequence* LevelSequence;
 
+	/** Event that is cast when this toolkit is closed */
+	FLevelSequenceEditorToolkitClosed OnClosedEvent;
+
 	/** The sequencer used by this editor. */
 	TSharedPtr<ISequencer> Sequencer;
 
@@ -133,9 +146,6 @@ private:
 	/** Pointer to the style set to use for toolkits. */
 	TSharedRef<ISlateStyle> Style;
 
-	/** Event that is cast when this toolkit is closed */
-	FLevelSequenceEditorToolkitClosed OnClosedEvent;
-
 private:
 
 	/**	The tab ids for all the tabs used */
diff --git a/Engine/Plugins/MovieScene/LevelSequenceEditor/Source/LevelSequenceEditor/Private/Misc/LevelSequenceEditorSpawnRegister.cpp b/Engine/Plugins/MovieScene/LevelSequenceEditor/Source/LevelSequenceEditor/Private/Misc/LevelSequenceEditorSpawnRegister.cpp
--- a/Engine/Plugins/MovieScene/LevelSequenceEditor/Source/LevelSequenceEditor/Private/Misc/LevelSequenceEditorSpawnRegister.cpp
+++ b/Engine/Plugins/MovieScene/LevelSequenceEditor/Source/LevelSequenceEditor/Private/Misc/LevelSequenceEditorSpawnRegister.cpp
@@ -5,8 +5,6 @@
 #include "ISequencerObjectChangeListener.h"
 #include "Editor/LevelEditor/Public/LevelEditor.h"
 #include "MovieScene.h"
-#include "Kismet2/KismetEditorUtilities.h"
-#include "Kismet2/BlueprintEditorUtils.h"
 #include "AssetSelection.h"
 
 #define LOCTEXT_NAMESPACE "LevelSequenceEditorSpawnRegister"
@@ -21,21 +19,25 @@ FLevelSequenceEditorSpawnRegister::FLevelSequenceEditorSpawnRegister()
 	FLevelEditorModule& LevelEditor = FModuleManager::GetModuleChecked<FLevelEditorModule>("LevelEditor");
 	OnActorSelectionChangedHandle = LevelEditor.OnActorSelectionChanged().AddRaw(this, &FLevelSequenceEditorSpawnRegister::HandleActorSelectionChanged);
 
-	OnActorMovedHandle = GEditor->OnActorMoved().AddLambda([=](AActor* Actor){
-		if (SpawnedObjects.Contains(Actor))
-		{
-			HandleAnyPropertyChanged(*Actor);
-		}
-	});
+	FAreObjectsEditable AreObjectsEditable = FAreObjectsEditable::CreateRaw(this, &FLevelSequenceEditorSpawnRegister::AreObjectsEditable);
+	OnAreObjectsEditableHandle = AreObjectsEditable.GetHandle();
+	LevelEditor.AddEditableObjectPredicate(AreObjectsEditable);
 }
 
 
 FLevelSequenceEditorSpawnRegister::~FLevelSequenceEditorSpawnRegister()
 {
-	GEditor->OnActorMoved().Remove(OnActorMovedHandle);
 	if (FLevelEditorModule* LevelEditor = FModuleManager::GetModulePtr<FLevelEditorModule>("LevelEditor"))
 	{
 		LevelEditor->OnActorSelectionChanged().Remove(OnActorSelectionChangedHandle);
+		LevelEditor->RemoveEditableObjectPredicate(OnAreObjectsEditableHandle);
+	}
+
+	TSharedPtr<ISequencer> Sequencer = WeakSequencer.Pin();
+	if (Sequencer.IsValid())
+	{
+		Sequencer->OnPreSave().RemoveAll(this);
+		Sequencer->OnActivateSequence().RemoveAll(this);
 	}
 }
 
@@ -46,25 +48,17 @@ FLevelSequenceEditorSpawnRegister::~FLevelSequenceEditorSpawnRegister()
 UObject* FLevelSequenceEditorSpawnRegister::SpawnObject(const FGuid& BindingId, FMovieSceneSequenceInstance& SequenceInstance, IMovieScenePlayer& Player)
 {
 	TGuardValue<bool> Guard(bShouldClearSelectionCache, false);
-	UObject* NewObject = FLevelSequenceSpawnRegister::SpawnObject(BindingId, SequenceInstance, Player);
+	AActor* NewObject = Cast<AActor>(FLevelSequenceSpawnRegister::SpawnObject(BindingId, SequenceInstance, Player));
 	
 	if (NewObject)
 	{
-		SpawnedObjects.Add(NewObject);
-
-		// Add an object listener for the spawned object to propagate changes back onto the spawnable default
-		TSharedPtr<ISequencer> Sequencer = WeakSequencer.Pin();
-		AActor* Actor = Cast<AActor>(NewObject);
+		// Cache the spawned object, and editable state first
+		SpawnedObjects.FindOrAdd(SequenceInstance.GetInstanceId()).Add(NewObject);
 
-		if (Sequencer.IsValid() && Actor)
+		// Select the actor if we think it should be selected
+		if (SelectedSpawnedObjects.Contains(FMovieSceneSpawnRegisterKey(BindingId, SequenceInstance)))
 		{
-			Sequencer->GetObjectChangeListener().GetOnAnyPropertyChanged(*NewObject).AddSP(this, &FLevelSequenceEditorSpawnRegister::HandleAnyPropertyChanged);
-
-			// Select the actor if we think it should be selected
-			if (Actor && SelectedSpawnedObjects.Contains(FMovieSceneSpawnRegisterKey(BindingId, SequenceInstance)))
-			{
-				GEditor->SelectActor(Actor, true /*bSelected*/, true /*bNotify*/);
-			}
+			GEditor->SelectActor(NewObject, true /*bSelected*/, true /*bNotify*/);
 		}
 	}
 
@@ -75,6 +69,14 @@ UObject* FLevelSequenceEditorSpawnRegister::SpawnObject(const FGuid& BindingId,
 void FLevelSequenceEditorSpawnRegister::PreDestroyObject(UObject& Object, const FGuid& BindingId, FMovieSceneSequenceInstance& SequenceInstance)
 {
 	TGuardValue<bool> Guard(bShouldClearSelectionCache, false);
+	
+	TSharedPtr<ISequencer> Sequencer = WeakSequencer.Pin();
+
+	// We only save default state for the currently focussed movie scene sequence instance
+	if (Sequencer.IsValid() && &Sequencer->GetFocusedMovieSceneSequenceInstance().Get() == &SequenceInstance)
+	{
+		SaveDefaultSpawnableState(BindingId, SequenceInstance);
+	}
 
 	// Cache its selection state
 	AActor* Actor = Cast<AActor>(&Object);
@@ -83,32 +85,79 @@ void FLevelSequenceEditorSpawnRegister::PreDestroyObject(UObject& Object, const
 		SelectedSpawnedObjects.Add(FMovieSceneSpawnRegisterKey(BindingId, SequenceInstance));
 		GEditor->SelectActor(Actor, false /*bSelected*/, true /*bNotify*/);
 	}
-			
-	SpawnedObjects.Remove(&Object);
+	
+	// Remove the spawned object (and anything that's null) from our cache
+	TSet<FObjectKey>* ExistingObjects = SpawnedObjects.Find(SequenceInstance.GetInstanceId());
+	if (ExistingObjects)
+	{
+		ExistingObjects->Remove(FObjectKey(&Object));
+		if (ExistingObjects->Num() == 0)
+		{
+			SpawnedObjects.Remove(SequenceInstance.GetInstanceId());
+		}
+	}
+
+	FLevelSequenceSpawnRegister::PreDestroyObject(Object, BindingId, SequenceInstance);
+}
 
-	// Remove our object listener
+void FLevelSequenceEditorSpawnRegister::SaveDefaultSpawnableState(const FGuid& BindingId, FMovieSceneSequenceInstance& SequenceInstance)
+{
 	TSharedPtr<ISequencer> Sequencer = WeakSequencer.Pin();
-	if (Sequencer.IsValid())
+	if (!Sequencer.IsValid())
 	{
-		Sequencer->GetObjectChangeListener().ReportObjectDestroyed(Object);
+		return;
 	}
 
-	FLevelSequenceSpawnRegister::PreDestroyObject(Object, BindingId, SequenceInstance);
+	UObject* Object = SequenceInstance.FindObject(BindingId, *Sequencer);
+	if (!Object)
+	{
+		return;
+	}
+
+	// Find the spawnable definition
+	UMovieScene* OwnerMovieScene = SequenceInstance.GetSequence()->GetMovieScene();
+	FMovieSceneSpawnable* Spawnable = OwnerMovieScene->FindSpawnable(BindingId);
+	if (Spawnable)
+	{
+		SequenceInstance.RestoreSpecificState(BindingId, *Sequencer);
+		Spawnable->CopyObjectTemplate(*Object, *OwnerMovieScene);
+	}
 }
 
+void FLevelSequenceEditorSpawnRegister::OnPreSaveMovieScene(ISequencer& InSequencer)
+{
+	// We're about to save the movie scene(s), so we need to save default spawnable state for the currently focused movie scene sequence instance
+
+	TSharedPtr<FMovieSceneSequenceInstance> Instance = InSequencer.GetFocusedMovieSceneSequenceInstance();
+	UMovieSceneSequence* Sequence = Instance->GetSequence();
+	UMovieScene* MovieScene = Sequence ? Sequence->GetMovieScene() : nullptr;
+	if (!MovieScene)
+	{
+		return;
+	}
+
+	for (int32 SpawnableIndex = 0; SpawnableIndex < MovieScene->GetSpawnableCount(); ++SpawnableIndex)
+	{
+		FMovieSceneSpawnable& Spawnable = MovieScene->GetSpawnable(SpawnableIndex);
+		SaveDefaultSpawnableState(Spawnable.GetGuid(), *Instance);
+	}
+}
+
+void FLevelSequenceEditorSpawnRegister::OnSequenceInstanceActivated(FMovieSceneSequenceInstance& ActiveInstance)
+{
+	ActiveSequence = ActiveInstance.GetInstanceId();
+}
 
 /* FLevelSequenceEditorSpawnRegister implementation
  *****************************************************************************/
 
-void FLevelSequenceEditorSpawnRegister::PopulateKeyedPropertyMap(AActor& SpawnedObject, TMap<UObject*, TSet<UProperty*>>& OutKeyedPropertyMap)
+void FLevelSequenceEditorSpawnRegister::SetSequencer(const TSharedPtr<ISequencer>& Sequencer)
 {
-	TSharedPtr<ISequencer> Sequencer = WeakSequencer.Pin();
-	Sequencer->GetAllKeyedProperties(SpawnedObject, OutKeyedPropertyMap.FindOrAdd(&SpawnedObject));
+	WeakSequencer = Sequencer;
+	Sequencer->OnPreSave().AddRaw(this, &FLevelSequenceEditorSpawnRegister::OnPreSaveMovieScene);
+	Sequencer->OnActivateSequence().AddRaw(this, &FLevelSequenceEditorSpawnRegister::OnSequenceInstanceActivated);
 
-	for (UActorComponent* Component : SpawnedObject.GetComponents())
-	{
-		Sequencer->GetAllKeyedProperties(*Component, OutKeyedPropertyMap.FindOrAdd(Component));
-	}
+	ActiveSequence = Sequencer->GetFocusedMovieSceneSequenceInstance()->GetInstanceId();
 }
 
 
@@ -123,61 +172,42 @@ void FLevelSequenceEditorSpawnRegister::HandleActorSelectionChanged(const TArray
 	}
 }
 
-
-void FLevelSequenceEditorSpawnRegister::HandleAnyPropertyChanged(UObject& SpawnedObject)
+bool FLevelSequenceEditorSpawnRegister::AreObjectsEditable(const TArray<TWeakObjectPtr<UObject>>& InObjects) const
 {
-	using namespace EditorUtilities;
-
-	AActor* Actor = CastChecked<AActor>(&SpawnedObject);
-	if (!Actor)
+	// Check that none of the objects specified are (or belong to) a spawned actor
+	for (const TWeakObjectPtr<UObject>& WeakObject : InObjects)
 	{
-		return;
-	}
-
-	TMap<UObject*, TSet<UProperty*>> ObjectToKeyedProperties;
-	PopulateKeyedPropertyMap(*Actor, ObjectToKeyedProperties);
-
-	// Copy any changed actor properties onto the default actor, provided they are not keyed
-	FCopyOptions Options(ECopyOptions::PropagateChangesToArchetypeInstances);
-
-	// Set up a property filter so only stuff that is not keyed gets copied onto the default
-	Options.PropertyFilter = [&](const UProperty& Property, const UObject& Object) -> bool {
-		const TSet<UProperty*>* ExcludedProperties = ObjectToKeyedProperties.Find(const_cast<UObject*>(&Object));
-
-		return !ExcludedProperties || !ExcludedProperties->Contains(const_cast<UProperty*>(&Property));
-	};
-
-	// Now copy the actor properties
-	AActor* DefaultActor = Actor->GetClass()->GetDefaultObject<AActor>();
-	EditorUtilities::CopyActorProperties(Actor, DefaultActor, Options);
-
-	// The above function call explicitly doesn't copy the root component transform (so the default actor is always at 0,0,0)
-	// But in sequencer, we want the object to have a default transform if it doesn't have a transform track
-	static FName RelativeLocation = GET_MEMBER_NAME_CHECKED(USceneComponent, RelativeLocation);
-	static FName RelativeRotation = GET_MEMBER_NAME_CHECKED(USceneComponent, RelativeRotation);
-	static FName RelativeScale3D = GET_MEMBER_NAME_CHECKED(USceneComponent, RelativeScale3D);
+		UObject* Object = WeakObject.Get();
+		if (!Object)
+		{
+			continue;
+		}
 
-	bool bHasKeyedTransform = false;
-	for (UProperty* Property : *ObjectToKeyedProperties.Find(Actor))
-	{
-		FName PropertyName = Property->GetFName();
-		bHasKeyedTransform = PropertyName == RelativeLocation || PropertyName == RelativeRotation || PropertyName == RelativeScale3D;
-		if (bHasKeyedTransform)
+		AActor* SourceActor = Cast<AActor>(Object);
+		if (!SourceActor)
 		{
-			break;
+			UActorComponent* ActorComponent = Cast<UActorComponent>(Object);
+			if (ActorComponent)
+			{
+				SourceActor = ActorComponent->GetOwner();
+			}
 		}
-	}
 
-	// Set the default transform if it's not keyed
-	USceneComponent* RootComponent = Actor->GetRootComponent();
-	USceneComponent* DefaultRootComponent = DefaultActor->GetRootComponent();
+		if (!SourceActor)
+		{
+			continue;
+		}
 
-	if (!bHasKeyedTransform && RootComponent && DefaultRootComponent)
-	{
-		DefaultRootComponent->RelativeLocation = RootComponent->RelativeLocation;
-		DefaultRootComponent->RelativeRotation = RootComponent->RelativeRotation;
-		DefaultRootComponent->RelativeScale3D = RootComponent->RelativeScale3D;
+		FObjectKey ThisObject(SourceActor);
+		for (auto& Pair : SpawnedObjects)
+		{
+			if (Pair.Key != ActiveSequence && Pair.Value.Contains(ThisObject))
+			{
+				return false;
+			}
+		}
 	}
+	return true;
 }
 
 
@@ -187,7 +217,7 @@ TValueOrError<FNewSpawnable, FText> FLevelSequenceEditorSpawnRegister::CreateNew
 {
 	FNewSpawnable NewSpawnable(nullptr, FName::NameToDisplayString(SourceObject.GetName(), false));
 
-	const FName BlueprintName = MakeUniqueObjectName(&OwnerMovieScene, UBlueprint::StaticClass(), SourceObject.GetFName());
+	const FName TemplateName = MakeUniqueObjectName(&OwnerMovieScene, UObject::StaticClass(), SourceObject.GetFName());
 
 	// First off, deal with creating a spawnable from a class
 	if (UClass* InClass = Cast<UClass>(&SourceObject))
@@ -198,41 +228,20 @@ TValueOrError<FNewSpawnable, FText> FLevelSequenceEditorSpawnRegister::CreateNew
 			return MakeError(ErrorText);
 		}
 
-		NewSpawnable.Blueprint = FKismetEditorUtilities::CreateBlueprint( InClass, &OwnerMovieScene, BlueprintName, EBlueprintType::BPTYPE_Normal, UBlueprint::StaticClass(), UBlueprintGeneratedClass::StaticClass() );
+		NewSpawnable.ObjectTemplate = NewObject<UObject>(&OwnerMovieScene, InClass, TemplateName);
 	}
 
 	// Deal with creating a spawnable from an instance of an actor
 	else if (AActor* Actor = Cast<AActor>(&SourceObject))
 	{
-		using namespace EditorUtilities;
-
-		UBlueprint* ActorBlueprint = Cast<UBlueprint>(Actor->GetClass()->ClassGeneratedBy);
-		if (ActorBlueprint)
-		{
-			// Create a new blueprint out of this actor's parent blueprint
-			NewSpawnable.Blueprint = FKismetEditorUtilities::CreateBlueprint(ActorBlueprint->GeneratedClass, &OwnerMovieScene, BlueprintName, EBlueprintType::BPTYPE_Normal, UBlueprint::StaticClass(), UBlueprintGeneratedClass::StaticClass());
-		}
-		else
-		{
-			NewSpawnable.Blueprint = FKismetEditorUtilities::CreateBlueprint(Actor->GetClass(), &OwnerMovieScene, BlueprintName, EBlueprintType::BPTYPE_Normal, UBlueprint::StaticClass(), UBlueprintGeneratedClass::StaticClass(), FName("CreateFromActor"));
-		}
-
-		// Use the actor name
+		NewSpawnable.ObjectTemplate = StaticDuplicateObject(Actor, &OwnerMovieScene, TemplateName);
 		NewSpawnable.Name = Actor->GetActorLabel();
 	}
 
 	// If it's a blueprint, we need some special handling
 	else if (UBlueprint* SourceBlueprint = Cast<UBlueprint>(&SourceObject))
 	{
-		UClass* SourceParentClass = SourceBlueprint->GeneratedClass;
-
-		if (!FKismetEditorUtilities::CanCreateBlueprintOfClass(SourceParentClass))
-		{
-			FText ErrorText = FText::Format(LOCTEXT("UnableToAddSpawnableBlueprint", "Unable to add spawnable for class of type '{0}' since it is not a valid blueprint parent class."), FText::FromString(SourceParentClass->GetName()));
-			return MakeError(ErrorText);
-		}
-
-		NewSpawnable.Blueprint = FKismetEditorUtilities::CreateBlueprint(SourceParentClass, &OwnerMovieScene, BlueprintName, EBlueprintType::BPTYPE_Normal, UBlueprint::StaticClass(), UBlueprintGeneratedClass::StaticClass());
+		NewSpawnable.ObjectTemplate = NewObject<UObject>(&OwnerMovieScene, SourceBlueprint->GeneratedClass, TemplateName);
 	}
 
 	// At this point we have to assume it's an asset
@@ -246,17 +255,32 @@ TValueOrError<FNewSpawnable, FText> FLevelSequenceEditorSpawnRegister::CreateNew
 			return MakeError(ErrorText);
 		}
 
-		NewSpawnable.Blueprint = FactoryToUse->CreateBlueprint( &SourceObject, &OwnerMovieScene, BlueprintName );
-	}
+		FText ErrorText;
+		if (!FactoryToUse->CanCreateActorFrom(FAssetData(&SourceObject), ErrorText))
+		{
+			if (!ErrorText.IsEmpty())
+			{
+				return MakeError(FText::Format(LOCTEXT("CannotCreateActorFromAsset_Ex", "Unable to create spawnable from  asset '{0}'. {1}."), FText::FromString(SourceObject.GetName()), ErrorText));
+			}
+			else
+			{
+				return MakeError(FText::Format(LOCTEXT("CannotCreateActorFromAsset", "Unable to create spawnable from  asset '{0}'."), FText::FromString(SourceObject.GetName())));
+			}
+		}
+
+		AActor* Instance = FactoryToUse->CreateActor(&SourceObject, GWorld->PersistentLevel, FTransform(), RF_Transactional, TemplateName );
 
+		NewSpawnable.ObjectTemplate = StaticDuplicateObject(Instance, &OwnerMovieScene, TemplateName);
+
+		GWorld->DestroyActor(Instance);
+	}
 
-	if (!NewSpawnable.Blueprint)
+	if (!NewSpawnable.ObjectTemplate || !NewSpawnable.ObjectTemplate->IsA<AActor>())
 	{
 		FText ErrorText = FText::Format(LOCTEXT("UnknownClassError", "Unable to create a new spawnable object from {0}."), FText::FromString(SourceObject.GetName()));
 		return MakeError(ErrorText);
 	}
 
-	FKismetEditorUtilities::CompileBlueprint(NewSpawnable.Blueprint);
 	return MakeValue(NewSpawnable);
 }
 
diff --git a/Engine/Plugins/MovieScene/LevelSequenceEditor/Source/LevelSequenceEditor/Private/Misc/LevelSequenceEditorSpawnRegister.h b/Engine/Plugins/MovieScene/LevelSequenceEditor/Source/LevelSequenceEditor/Private/Misc/LevelSequenceEditorSpawnRegister.h
--- a/Engine/Plugins/MovieScene/LevelSequenceEditor/Source/LevelSequenceEditor/Private/Misc/LevelSequenceEditorSpawnRegister.h
+++ b/Engine/Plugins/MovieScene/LevelSequenceEditor/Source/LevelSequenceEditor/Private/Misc/LevelSequenceEditorSpawnRegister.h
@@ -29,10 +29,7 @@ public:
 
 public:
 
-	void SetSequencer(const TSharedPtr<ISequencer>& Sequencer)
-	{
-		WeakSequencer = Sequencer;
-	}
+	void SetSequencer(const TSharedPtr<ISequencer>& Sequencer);
 
 public:
 
@@ -44,33 +41,40 @@ public:
 	virtual TValueOrError<FNewSpawnable, FText> CreateNewSpawnableType(UObject& SourceObject, UMovieScene& OwnerMovieScene) override;
 #endif
 
-protected:
-
-	/** Populate a map of properties that are keyed on a particular object. */
-	void PopulateKeyedPropertyMap(AActor& SpawnedObject, TMap<UObject*, TSet<UProperty*>>& OutKeyedPropertyMap);
-
 private:
 
 	/** Called when the editor selection has changed. */
 	void HandleActorSelectionChanged(const TArray<UObject*>& NewSelection, bool bForceRefresh);
 
-	/** Called when a property on a spawned object changes. */
-	void HandleAnyPropertyChanged(UObject& SpawnedObject);
+	/** Called before sequencer attempts to save the movie scene(s) it's editing */
+	void OnPreSaveMovieScene(ISequencer& InSequencer);
+
+	/** Called when a new movie scene sequence instance has been activated */
+	void OnSequenceInstanceActivated(FMovieSceneSequenceInstance& ActiveInstance);
+
+	/** Saves the default state for the specified spawnable, if an instance for it currently exists */
+	void SaveDefaultSpawnableState(const FGuid& BindingId, FMovieSceneSequenceInstance& SequenceInstance);
+
+	/** Check whether the specified objects are editable on the details panel. Called from the level editor */
+	bool AreObjectsEditable(const TArray<TWeakObjectPtr<UObject>>& InObjects) const;
 
 private:
 
 	/** Handles for delegates that we've bound to. */
-	FDelegateHandle OnActorMovedHandle, OnActorSelectionChangedHandle;
+	FDelegateHandle OnActorSelectionChangedHandle, OnAreObjectsEditableHandle;
 
 	/** Set of spawn register keys for objects that should be selected if they are spawned. */
 	TSet<FMovieSceneSpawnRegisterKey> SelectedSpawnedObjects;
 
 	/** Set of currently spawned objects */
-	TSet<FObjectKey> SpawnedObjects;
+	TMap<FGuid, TSet<FObjectKey>> SpawnedObjects;
 
 	/** True if we should clear the above selection cache when the editor selection has been changed. */
 	bool bShouldClearSelectionCache;
 
 	/** Weak pointer to the active sequencer. */
 	TWeakPtr<ISequencer> WeakSequencer;
+
+	/** Identifier for the current active level sequence */
+	FGuid ActiveSequence;
 };
diff --git a/Engine/Plugins/MovieScene/MatineeToLevelSequence/MatineeToLevelSequence.uplugin b/Engine/Plugins/MovieScene/MatineeToLevelSequence/MatineeToLevelSequence.uplugin
new file mode 100644
--- /dev/null
+++ b/Engine/Plugins/MovieScene/MatineeToLevelSequence/MatineeToLevelSequence.uplugin
@@ -0,0 +1,25 @@
+{
+	"FileVersion" : 3,
+	"Version" : 1,
+	"VersionName" : "1.0",
+	"FriendlyName" : "Matinee to Level Sequence",
+	"Description" : "Convert Matinee to LevelSequence Asset.",
+	"Category" : "Editor",
+	"CreatedBy" : "Epic Games, Inc.",
+	"CreatedByURL" : "http://epicgames.com",
+	"DocsURL" : "",
+	"MarketplaceURL" : "",
+	"SupportURL" : "",
+	"EnabledByDefault" : true,
+	"CanContainContent" : false,
+	"IsBetaVersion" : false,
+	"Installed" : false,
+	"Modules" :
+	[
+		{
+			"Name" : "MatineeToLevelSequence",
+			"Type" : "Editor",
+			"LoadingPhase" : "PostEngineInit"
+		}
+	]
+}
\ No newline at end of file
diff --git a/Engine/Plugins/MovieScene/MatineeToLevelSequence/Source/MatineeToLevelSequence/MatineeToLevelSequence.Build.cs b/Engine/Plugins/MovieScene/MatineeToLevelSequence/Source/MatineeToLevelSequence/MatineeToLevelSequence.Build.cs
new file mode 100644
--- /dev/null
+++ b/Engine/Plugins/MovieScene/MatineeToLevelSequence/Source/MatineeToLevelSequence/MatineeToLevelSequence.Build.cs
@@ -0,0 +1,40 @@
+// Copyright 1998-2016 Epic Games, Inc. All Rights Reserved.
+
+using UnrealBuildTool;
+
+public class MatineeToLevelSequence : ModuleRules
+{
+	public MatineeToLevelSequence(TargetInfo Target)
+	{
+        DynamicallyLoadedModuleNames.AddRange(
+            new string[] {
+				"AssetTools",
+			}
+        );
+        
+        PrivateDependencyModuleNames.AddRange(
+			new string[] {
+                "LevelSequence",
+				"Core",
+				"CoreUObject",
+                "EditorStyle",
+                "Engine",
+				"LevelEditor",
+				"MovieScene",
+                "MovieSceneTools",
+				"MovieSceneTracks",
+                "Slate",
+                "SlateCore",
+                "UnrealEd",
+            }
+        );
+
+		PrivateIncludePathModuleNames.AddRange(
+			new string[] {
+				"AssetTools",
+                "MovieSceneTools",
+                "Settings",
+			}
+		);
+	}
+}
diff --git a/Engine/Plugins/MovieScene/MatineeToLevelSequence/Source/MatineeToLevelSequence/Private/MatineeToLevelSequenceModule.cpp b/Engine/Plugins/MovieScene/MatineeToLevelSequence/Source/MatineeToLevelSequence/Private/MatineeToLevelSequenceModule.cpp
new file mode 100644
--- /dev/null
+++ b/Engine/Plugins/MovieScene/MatineeToLevelSequence/Source/MatineeToLevelSequence/Private/MatineeToLevelSequenceModule.cpp
@@ -0,0 +1,521 @@
+// Copyright 1998-2016 Epic Games, Inc. All Rights Reserved.
+
+#include "MatineeToLevelSequencePCH.h"
+#include "ModuleInterface.h"
+
+#define LOCTEXT_NAMESPACE "MatineeToLevelSequence"
+
+DEFINE_LOG_CATEGORY(LogMatineeToLevelSequence);
+
+/**
+ * Implements the MatineeToLevelSequence module.
+ */
+class FMatineeToLevelSequenceModule
+	: public IModuleInterface
+{
+public:
+
+	// IModuleInterface interface
+
+	virtual void StartupModule() override
+	{
+		RegisterMenuExtensions();
+	}
+	
+	virtual void ShutdownModule() override
+	{
+		UnregisterMenuExtensions();
+	}
+
+protected:
+
+	/** Register menu extensions for the level editor toolbar. */
+	void RegisterMenuExtensions()
+	{
+		// Register level editor menu extender
+		LevelEditorMenuExtenderDelegate = FLevelEditorModule::FLevelViewportMenuExtender_SelectedActors::CreateStatic(&FMatineeToLevelSequenceModule::ExtendLevelViewportContextMenu);
+		FLevelEditorModule& LevelEditorModule = FModuleManager::LoadModuleChecked<FLevelEditorModule>(TEXT("LevelEditor"));
+		auto& MenuExtenders = LevelEditorModule.GetAllLevelViewportContextMenuExtenders();
+		MenuExtenders.Add(LevelEditorMenuExtenderDelegate);
+		LevelEditorExtenderDelegateHandle = MenuExtenders.Last().GetHandle();
+	}
+
+protected:
+
+	/** Unregisters menu extensions for the level editor toolbar. */
+	void UnregisterMenuExtensions()
+	{
+		// Unregister level editor menu extender
+		if (FModuleManager::Get().IsModuleLoaded(TEXT("LevelEditor")))
+		{
+			FLevelEditorModule& LevelEditorModule = FModuleManager::GetModuleChecked<FLevelEditorModule>(TEXT("LevelEditor"));
+			LevelEditorModule.GetAllLevelViewportContextMenuExtenders().RemoveAll([&](const FLevelEditorModule::FLevelViewportMenuExtender_SelectedActors& Delegate) {
+				return Delegate.GetHandle() == LevelEditorExtenderDelegateHandle;
+			});
+		}
+	}
+
+	static TSharedRef<FExtender> ExtendLevelViewportContextMenu(const TSharedRef<FUICommandList> CommandList, const TArray<AActor*> SelectedActors)
+	{
+		TSharedRef<FExtender> Extender(new FExtender());
+
+		TArray<TWeakObjectPtr<AActor> > ActorsToConvert;
+		for (AActor* Actor : SelectedActors)
+		{
+			if (Actor->IsA(AMatineeActor::StaticClass()))
+			{
+				ActorsToConvert.Add(Actor);
+			}
+		}
+
+		if (ActorsToConvert.Num())
+		{
+			// Add the convert to level sequence asset sub-menu extender
+			Extender->AddMenuExtension(
+				"ActorSelectVisibilityLevels",
+				EExtensionHook::After,
+				nullptr,
+				FMenuExtensionDelegate::CreateStatic(&FMatineeToLevelSequenceModule::CreateLevelViewportContextMenuEntries, ActorsToConvert));
+		}
+
+		return Extender;
+	}
+
+	static void CreateLevelViewportContextMenuEntries(FMenuBuilder& MenuBuilder, TArray<TWeakObjectPtr<AActor> > ActorsToConvert)
+	{
+		MenuBuilder.BeginSection("LevelSequence", LOCTEXT("LevelSequenceLevelEditorHeading", "Level Sequence"));
+
+		MenuBuilder.AddMenuEntry(
+			LOCTEXT("MenuExtensionConvertMatineeToLevelSequence", "Convert to Level Sequence"),
+			LOCTEXT("MenuExtensionConvertMatineeToLevelSequence_Tooltip", "Convert to Level Sequence"),
+			FSlateIcon(),
+			FExecuteAction::CreateStatic(&FMatineeToLevelSequenceModule::OnConvertMatineeToLevelSequence, ActorsToConvert),
+			NAME_None,
+			EUserInterfaceActionType::Button);
+
+		MenuBuilder.EndSection();
+	}
+
+	/** Callback for converting a matinee to a level sequence asset. */
+	static void OnConvertMatineeToLevelSequence(TArray<TWeakObjectPtr<AActor> > ActorsToConvert)
+	{
+		TArray<TWeakObjectPtr<ALevelSequenceActor> > NewActors;
+
+		for (TWeakObjectPtr<AActor> Actor : ActorsToConvert)
+		{
+			TWeakObjectPtr<ALevelSequenceActor> NewActor = ConvertSingleMatineeToLevelSequence(Actor);
+			if (NewActor.IsValid())
+			{
+				NewActors.Add(NewActor);
+			}
+		}
+
+		// Select the newly created level sequence actors
+		const bool bNotifySelectionChanged = true;
+		const bool bDeselectBSP = true;
+		const bool bWarnAboutTooManyActors = false;
+		const bool bSelectEvenIfHidden = false;
+
+		GEditor->GetSelectedActors()->Modify();
+		GEditor->GetSelectedActors()->BeginBatchSelectOperation();
+
+		GEditor->SelectNone( bNotifySelectionChanged, bDeselectBSP, bWarnAboutTooManyActors );
+
+		for (TWeakObjectPtr<AActor> NewActor : NewActors )
+		{
+			GEditor->SelectActor(NewActor.Get(), true, bNotifySelectionChanged, bSelectEvenIfHidden );
+		}
+
+		GEditor->GetSelectedActors()->EndBatchSelectOperation();
+		GEditor->NoteSelectionChange();
+
+		// Edit the first asset
+		if (NewActors.Num())
+		{
+			UObject* NewAsset = NewActors[0]->LevelSequence.TryLoad();
+			if (NewAsset)
+			{
+				FAssetEditorManager::Get().OpenEditorForAsset(NewAsset);
+			}
+		}
+	}
+
+	/** Find or add a folder for the given actor **/
+	static void FindOrAddFolder(UMovieScene* MovieScene, TWeakObjectPtr<AActor> Actor, FGuid Guid)
+	{
+		FName FolderName(NAME_None);
+		if(Actor.Get()->IsA<ACharacter>() || Actor.Get()->IsA<ASkeletalMeshActor>())
+		{
+			FolderName = TEXT("Characters");
+		}
+		else if(Actor.Get()->GetClass()->IsChildOf(ACameraActor::StaticClass()))
+		{
+			FolderName = TEXT("Cameras");
+		}
+		else if(Actor.Get()->GetClass()->IsChildOf(ALight::StaticClass()))
+		{
+			FolderName = TEXT("Lights");
+		}
+		else if (Actor.Get()->GetComponentsByClass(UParticleSystemComponent::StaticClass()).Num())
+		{
+			FolderName = TEXT("Particles");
+		}
+		else
+		{
+			FolderName = TEXT("Misc");
+		}
+
+		// look for a folder to put us in
+		UMovieSceneFolder* FolderToUse = nullptr;
+		for (UMovieSceneFolder* Folder : MovieScene->GetRootFolders())
+		{
+			if (Folder->GetFolderName() == FolderName)
+			{
+				FolderToUse = Folder;
+				break;
+			}
+		}
+
+		if (FolderToUse == nullptr)
+		{
+			FolderToUse = NewObject<UMovieSceneFolder>(MovieScene, NAME_None, RF_Transactional);
+			FolderToUse->SetFolderName(FolderName);
+			MovieScene->GetRootFolders().Add(FolderToUse);
+		}
+
+		FolderToUse->AddChildObjectBinding(Guid);
+	}
+
+	/** Add master track to a folder **/
+	static void AddMasterTrackToFolder(UMovieScene* MovieScene, UMovieSceneTrack* MovieSceneTrack, FName FolderName)
+	{
+		UMovieSceneFolder* FolderToUse = nullptr;
+		for (UMovieSceneFolder* Folder : MovieScene->GetRootFolders())
+		{
+			if (Folder->GetFolderName() == FolderName)
+			{
+				FolderToUse = Folder;
+				break;
+			}
+		}
+
+		if (FolderToUse == nullptr)
+		{
+			FolderToUse = NewObject<UMovieSceneFolder>(MovieScene, NAME_None, RF_Transactional);
+			FolderToUse->SetFolderName(FolderName);
+			MovieScene->GetRootFolders().Add(FolderToUse);
+		}
+
+		FolderToUse->AddChildMasterTrack(MovieSceneTrack);
+	}
+
+	/** Add property to possessable node **/
+	template <typename T>
+	static T* AddPropertyTrack(FName InPropertyName, AActor* InActor, const FGuid& PossessableGuid, UMovieSceneSequence* NewSequence, UMovieScene* NewMovieScene)
+	{
+		T* PropertyTrack = nullptr;
+
+		// Find the property that matinee uses
+		void* PropContainer = NULL;
+		UProperty* Property = NULL;
+		UObject* PropObject = FMatineeUtils::FindObjectAndPropOffset(PropContainer, Property, InActor, InPropertyName );
+
+		FGuid ObjectGuid = PossessableGuid;
+		if (PropObject && Property)
+		{
+			// If the property object that owns this property isn't already bound, add a binding to the property object
+			ObjectGuid = NewSequence->FindPossessableObjectId(*PropObject);
+			if (!ObjectGuid.IsValid())
+			{
+				ObjectGuid = NewMovieScene->AddPossessable(PropObject->GetName(), PropObject->GetClass());
+				NewSequence->BindPossessableObject(ObjectGuid, *PropObject, GWorld);
+			}
+
+			FString PropertyName = Property->GetFName().ToString();
+			FString PropertyPath = Property->GetPathName();
+
+			// Special case for Light components which have some deprecated names
+			if (PropObject->GetClass()->IsChildOf(ULightComponentBase::StaticClass()))
+			{
+				TMap<FName, FName> PropertyNameRemappings;
+				PropertyNameRemappings.Add(TEXT("Brightness"), TEXT("Intensity"));
+				PropertyNameRemappings.Add(TEXT("Radius"), TEXT("AttenuationRadius"));
+
+				FName* RemappedName = PropertyNameRemappings.Find(*PropertyName);
+				if (RemappedName != nullptr)
+				{
+					PropertyName = RemappedName->ToString();
+					int32 LastDot = INDEX_NONE;
+					PropertyPath.FindLastChar(TCHAR(':'), LastDot);
+					if (LastDot != INDEX_NONE)
+					{
+						PropertyPath = PropertyPath.Left(LastDot) + TEXT(":") + PropertyName;
+					}
+				}
+			}
+
+			PropertyTrack = NewMovieScene->AddTrack<T>(ObjectGuid);	
+			PropertyTrack->SetPropertyNameAndPath(*PropertyName, PropertyPath);
+		}
+		else
+		{
+			UE_LOG(LogMatineeToLevelSequence, Log, TEXT("Can't find property '%s' for '%s'."), *InPropertyName.ToString(), *InActor->GetActorLabel());
+		}
+
+		return PropertyTrack;
+	}
+
+	/** Convert an interp group */
+	static void ConvertInterpGroup(UInterpGroup* Group, AActor* GroupActor, UMovieScene* NewMovieScene, UMovieSceneSequence* NewSequence)
+	{
+		FGuid PossessableGuid;
+
+		// Bind the group actor as a possessable						
+		if (GroupActor)
+		{
+			PossessableGuid = NewMovieScene->AddPossessable(GroupActor->GetActorLabel(), GroupActor->GetClass());
+			NewSequence->BindPossessableObject(PossessableGuid, *GroupActor, GWorld);
+	
+			FindOrAddFolder(NewMovieScene, GroupActor, PossessableGuid);
+		}
+
+		for (int32 j=0; j<Group->InterpTracks.Num(); ++j)
+		{
+			UInterpTrack* Track = Group->InterpTracks[j];
+			if (Track->IsDisabled())
+			{
+				continue;
+			}
+
+			// Handle each track class
+			if (Track->IsA(UInterpTrackMove::StaticClass()))					
+			{
+				UInterpTrackMove* MatineeMoveTrack = StaticCast<UInterpTrackMove*>(Track);
+				if (MatineeMoveTrack->GetNumKeyframes() != 0 && PossessableGuid.IsValid())
+				{
+					UMovieScene3DTransformTrack* TransformTrack = NewMovieScene->AddTrack<UMovieScene3DTransformTrack>(PossessableGuid);								
+					FMatineeImportTools::CopyInterpMoveTrack(MatineeMoveTrack, TransformTrack);
+				}
+			}
+			else if (Track->IsA(UInterpTrackAnimControl::StaticClass()))
+			{
+				UInterpTrackAnimControl* MatineeAnimControlTrack = StaticCast<UInterpTrackAnimControl*>(Track);
+				if (MatineeAnimControlTrack->GetNumKeyframes() != 0 && PossessableGuid.IsValid())
+				{
+					UMovieSceneSkeletalAnimationTrack* SkeletalAnimationTrack = NewMovieScene->AddTrack<UMovieSceneSkeletalAnimationTrack>(PossessableGuid);	
+					float EndPlaybackRange = NewMovieScene->GetPlaybackRange().GetUpperBoundValue();
+					FMatineeImportTools::CopyInterpAnimControlTrack(MatineeAnimControlTrack, SkeletalAnimationTrack, EndPlaybackRange);
+				}
+			}
+			else if (Track->IsA(UInterpTrackToggle::StaticClass()))
+			{
+				UInterpTrackToggle* MatineeParticleTrack = StaticCast<UInterpTrackToggle*>(Track);
+				if (MatineeParticleTrack->GetNumKeyframes() != 0 && PossessableGuid.IsValid())
+				{
+					UMovieSceneParticleTrack* ParticleTrack = NewMovieScene->AddTrack<UMovieSceneParticleTrack>(PossessableGuid);	
+					FMatineeImportTools::CopyInterpParticleTrack(MatineeParticleTrack, ParticleTrack);
+				}
+			}
+			else if (Track->IsA(UInterpTrackEvent::StaticClass()))
+			{
+				UInterpTrackEvent* MatineeEventTrack = StaticCast<UInterpTrackEvent*>(Track);
+				if (MatineeEventTrack->GetNumKeyframes() != 0)
+				{
+					UMovieSceneEventTrack* EventTrack = Cast<UMovieSceneEventTrack>(NewMovieScene->AddMasterTrack<UMovieSceneEventTrack>());
+					FString EventTrackName = Group->GroupName.ToString() + TEXT("Events");
+					EventTrack->SetDisplayName(FText::FromString(EventTrackName));
+					FMatineeImportTools::CopyInterpEventTrack(MatineeEventTrack, EventTrack);
+
+					static FName EventsFolder("Events");
+					AddMasterTrackToFolder(NewMovieScene, EventTrack, EventsFolder);
+				}
+			}
+			else if (Track->IsA(UInterpTrackSound::StaticClass()))
+			{
+				UInterpTrackSound* MatineeSoundTrack = StaticCast<UInterpTrackSound*>(Track);
+				if (MatineeSoundTrack->GetNumKeyframes() != 0)
+				{
+					UMovieSceneAudioTrack* AudioTrack = Cast<UMovieSceneAudioTrack>(NewMovieScene->AddMasterTrack<UMovieSceneAudioTrack>());
+					FString AudioTrackName = Group->GroupName.ToString() + TEXT("Audio");
+					AudioTrack->SetDisplayName(FText::FromString(AudioTrackName));					
+					FMatineeImportTools::CopyInterpSoundTrack(MatineeSoundTrack, AudioTrack);
+
+					static FName AudioFolder("Audio");
+					AddMasterTrackToFolder(NewMovieScene, AudioTrack, AudioFolder);
+				}
+			}
+			else if (Track->IsA(UInterpTrackBoolProp::StaticClass()))
+			{
+				UInterpTrackBoolProp* MatineeBoolTrack = StaticCast<UInterpTrackBoolProp*>(Track);
+				if (MatineeBoolTrack->GetNumKeyframes() != 0 && GroupActor && PossessableGuid.IsValid())
+				{
+					UMovieSceneBoolTrack* BoolTrack = AddPropertyTrack<UMovieSceneBoolTrack>(MatineeBoolTrack->PropertyName, GroupActor, PossessableGuid, NewSequence, NewMovieScene);
+					if (BoolTrack)
+					{
+						FMatineeImportTools::CopyInterpBoolTrack(MatineeBoolTrack, BoolTrack);
+					}
+				}
+			}
+			else if (Track->IsA(UInterpTrackFloatProp::StaticClass()))
+			{
+				UInterpTrackFloatProp* MatineeFloatTrack = StaticCast<UInterpTrackFloatProp*>(Track);
+				if (MatineeFloatTrack->GetNumKeyframes() != 0 && GroupActor && PossessableGuid.IsValid())
+				{
+					UMovieSceneFloatTrack* FloatTrack = AddPropertyTrack<UMovieSceneFloatTrack>(MatineeFloatTrack->PropertyName, GroupActor, PossessableGuid, NewSequence, NewMovieScene);
+					if (FloatTrack)
+					{
+						FMatineeImportTools::CopyInterpFloatTrack(MatineeFloatTrack, FloatTrack);
+					}
+				}
+			}
+			else if (Track->IsA(UInterpTrackColorProp::StaticClass()))
+			{
+				UInterpTrackColorProp* MatineeColorTrack = StaticCast<UInterpTrackColorProp*>(Track);
+				if (MatineeColorTrack->GetNumKeyframes() != 0 && GroupActor && PossessableGuid.IsValid())
+				{
+					UMovieSceneColorTrack* ColorTrack = AddPropertyTrack<UMovieSceneColorTrack>(MatineeColorTrack->PropertyName, GroupActor, PossessableGuid, NewSequence, NewMovieScene);
+					if (ColorTrack)
+					{
+						FMatineeImportTools::CopyInterpColorTrack(MatineeColorTrack, ColorTrack);
+					}
+				}
+			}
+			else if (Track->IsA(UInterpTrackVisibility::StaticClass()))
+			{
+				UInterpTrackVisibility* MatineeVisibilityTrack = StaticCast<UInterpTrackVisibility*>(Track);
+				if (MatineeVisibilityTrack->GetNumKeyframes() != 0 && GroupActor && PossessableGuid.IsValid())
+				{
+					UMovieSceneVisibilityTrack* VisibilityTrack = NewMovieScene->AddTrack<UMovieSceneVisibilityTrack>(PossessableGuid);	
+					if (VisibilityTrack)
+					{
+						VisibilityTrack->SetPropertyNameAndPath(TEXT("bHidden"), GroupActor->GetPathName() + TEXT(".bHidden"));
+
+						FMatineeImportTools::CopyInterpVisibilityTrack(MatineeVisibilityTrack, VisibilityTrack);
+					}
+				}
+			}
+			else if (Track->IsA(UInterpTrackDirector::StaticClass()))
+			{
+				// Intentionally left blank - The director track is converted in a separate pass below.
+			}
+			else
+			{
+				if (GroupActor)
+				{
+					UE_LOG(LogMatineeToLevelSequence, Log, TEXT("Unsupported track '%s' for '%s'."), *Track->TrackTitle, *GroupActor->GetActorLabel());
+				}
+				else
+				{
+					UE_LOG(LogMatineeToLevelSequence, Log, TEXT("Unsupported track '%s'."), *Track->TrackTitle);
+				}
+			}
+		}
+	}
+
+	/** Convert a single matinee to a level sequence asset */
+	static TWeakObjectPtr<ALevelSequenceActor> ConvertSingleMatineeToLevelSequence(TWeakObjectPtr<AActor> ActorToConvert)
+	{
+		UObject* AssetOuter = ActorToConvert->GetOuter();
+		UPackage* AssetPackage = AssetOuter->GetOutermost();
+
+		FString NewLevelSequenceAssetName = ActorToConvert->GetActorLabel() + FString("LevelSequence");
+		FString NewLevelSequenceAssetPath = AssetPackage->GetName();
+		int LastSlashPos = NewLevelSequenceAssetPath.Find(TEXT("/"), ESearchCase::IgnoreCase, ESearchDir::FromEnd);
+		NewLevelSequenceAssetPath = NewLevelSequenceAssetPath.Left(LastSlashPos);
+
+		// Create a new level sequence asset with the appropriate name
+		IAssetTools& AssetTools = FModuleManager::GetModuleChecked<FAssetToolsModule>("AssetTools").Get();
+
+		UObject* NewAsset = nullptr;
+		for (TObjectIterator<UClass> It ; It ; ++It)
+		{
+			UClass* CurrentClass = *It;
+			if (CurrentClass->IsChildOf(UFactory::StaticClass()) && !(CurrentClass->HasAnyClassFlags(CLASS_Abstract)))
+			{
+				UFactory* Factory = Cast<UFactory>(CurrentClass->GetDefaultObject());
+				if (Factory->CanCreateNew() && Factory->ImportPriority >= 0 && Factory->SupportedClass == ULevelSequence::StaticClass())
+				{
+					NewAsset = AssetTools.CreateAssetWithDialog(NewLevelSequenceAssetName, NewLevelSequenceAssetPath, ULevelSequence::StaticClass(), Factory);
+					break;
+				}
+			}
+		}
+
+		if (!NewAsset)
+		{
+			return nullptr;
+		}
+
+		UMovieSceneSequence* NewSequence = Cast<UMovieSceneSequence>(NewAsset);
+		UMovieScene* NewMovieScene = NewSequence->GetMovieScene();
+
+		// Add a level sequence actor for this new sequence
+		UActorFactory* ActorFactory = GEditor->FindActorFactoryForActorClass(ALevelSequenceActor::StaticClass());
+		if (!ensure(ActorFactory))
+		{
+			return nullptr;
+		}
+
+		ALevelSequenceActor* NewActor = CastChecked<ALevelSequenceActor>(GEditor->UseActorFactory(ActorFactory, FAssetData(NewAsset), &FTransform::Identity));
+
+		// Walk through all the interp group data and create corresponding tracks on the new level sequence asset
+		if (ActorToConvert->IsA(AMatineeActor::StaticClass()))
+		{
+			AMatineeActor* MatineeActor = StaticCast<AMatineeActor*>(ActorToConvert.Get());
+			MatineeActor->InitInterp();
+
+			// Set the length
+			NewMovieScene->SetPlaybackRange(0.0f, MatineeActor->MatineeData->InterpLength);
+
+			// Convert the groups
+			for (int32 i=0; i<MatineeActor->GroupInst.Num(); ++i)
+			{
+				UInterpGroupInst* GrInst = MatineeActor->GroupInst[i];
+				UInterpGroup* Group = GrInst->Group;
+				AActor* GroupActor = GrInst->GetGroupActor();
+				ConvertInterpGroup(Group, GroupActor, NewMovieScene, NewSequence);
+			}
+
+			// Director group - convert this after the regular groups to ensure that the camera cut bindings are there
+			UInterpGroupDirector* DirGroup = MatineeActor->MatineeData->FindDirectorGroup();
+			if (DirGroup)
+			{
+				UInterpTrackDirector* MatineeDirectorTrack = DirGroup->GetDirectorTrack();
+				if (MatineeDirectorTrack && MatineeDirectorTrack->GetNumKeyframes() != 0)
+				{
+					UMovieSceneCameraCutTrack* CameraCutTrack = Cast<UMovieSceneCameraCutTrack>(NewMovieScene->AddMasterTrack<UMovieSceneCameraCutTrack>());
+					FMatineeImportTools::CopyInterpDirectorTrack(MatineeDirectorTrack, CameraCutTrack, MatineeActor, NewSequence);
+				}
+
+				UInterpTrackFade* MatineeFadeTrack = DirGroup->GetFadeTrack();
+				if (MatineeFadeTrack && MatineeFadeTrack->GetNumKeyframes() != 0)
+				{						
+					UMovieSceneFadeTrack* FadeTrack = Cast<UMovieSceneFadeTrack>(NewMovieScene->AddMasterTrack<UMovieSceneFadeTrack>());
+					FMatineeImportTools::CopyInterpFadeTrack(MatineeFadeTrack, FadeTrack);
+				}
+
+				UInterpTrackSlomo* MatineeSlomoTrack = DirGroup->GetSlomoTrack();
+				if (MatineeSlomoTrack && MatineeSlomoTrack->GetNumKeyframes() != 0)
+				{
+					UMovieSceneSlomoTrack* SlomoTrack = Cast<UMovieSceneSlomoTrack>(NewMovieScene->AddMasterTrack<UMovieSceneSlomoTrack>());
+					FMatineeImportTools::CopyInterpFloatTrack(MatineeSlomoTrack, SlomoTrack);
+				}
+				
+				//@todo - color scale, audio master
+			}
+
+			MatineeActor->TermInterp();
+		}
+		return NewActor;
+	}
+
+private:
+
+	FLevelEditorModule::FLevelViewportMenuExtender_SelectedActors LevelEditorMenuExtenderDelegate;
+
+	FDelegateHandle LevelEditorExtenderDelegateHandle;
+};
+
+IMPLEMENT_MODULE(FMatineeToLevelSequenceModule, MatineeToLevelSequence);
+
+#undef LOCTEXT_NAMESPACE
diff --git a/Engine/Plugins/MovieScene/MatineeToLevelSequence/Source/MatineeToLevelSequence/Private/MatineeToLevelSequencePCH.h b/Engine/Plugins/MovieScene/MatineeToLevelSequence/Source/MatineeToLevelSequence/Private/MatineeToLevelSequencePCH.h
new file mode 100644
--- /dev/null
+++ b/Engine/Plugins/MovieScene/MatineeToLevelSequence/Source/MatineeToLevelSequence/Private/MatineeToLevelSequencePCH.h
@@ -0,0 +1,70 @@
+// Copyright 1998-2016 Epic Games, Inc. All Rights Reserved.
+
+#pragma once
+
+
+/* Private dependencies
+ *****************************************************************************/
+
+#include "Core.h"
+#include "CoreUObject.h"
+#include "LevelSequence.h"
+#include "AssetToolsModule.h"
+#include "EditorStyle.h"
+#include "Engine.h"
+
+
+/* Private includes
+ *****************************************************************************/
+
+#include "Matinee/InterpData.h"
+#include "Matinee/InterpTrackFloatBase.h"
+#include "Matinee/InterpTrackBoolProp.h"
+#include "Matinee/InterpTrackColorProp.h"
+#include "Matinee/InterpTrackFloatProp.h"
+#include "Matinee/InterpTrackMove.h"
+#include "Matinee/InterpTrackAnimControl.h"
+#include "Matinee/InterpTrackSound.h"
+#include "Matinee/InterpTrackFade.h"
+#include "Matinee/InterpTrackSlomo.h"
+#include "Matinee/InterpTrackDirector.h"
+#include "Matinee/InterpTrackEvent.h"
+#include "Matinee/InterpTrackAudioMaster.h"
+#include "Matinee/InterpTrackVisibility.h"
+#include "Matinee/InterpGroup.h"
+#include "Matinee/InterpGroupInst.h"
+#include "Matinee/InterpGroupDirector.h"
+#include "Matinee/MatineeActor.h"
+#include "MatineeUtils.h"
+
+#include "LevelEditor.h"
+#include "LevelSequenceActor.h"
+
+#include "MatineeImportTools.h"
+#include "MovieScene.h"
+#include "MovieSceneFolder.h"
+#include "MovieSceneBoolTrack.h"
+#include "MovieSceneColorTrack.h"
+#include "MovieSceneFloatTrack.h"
+#include "MovieScene3DTransformTrack.h"
+#include "MovieSceneParticleTrack.h"
+#include "MovieSceneSkeletalAnimationTrack.h"
+#include "MovieSceneAudioTrack.h"
+#include "MovieSceneFadeTrack.h"
+#include "MovieSceneSlomoTrack.h"
+#include "MovieSceneCameraCutTrack.h"
+#include "MovieSceneEventTrack.h"
+#include "MovieSceneAudioTrack.h"
+#include "MovieSceneVisibilityTrack.h"
+
+#include "AssetToolsModule.h"
+#include "AssetRegistryModule.h"
+#include "Camera/CameraActor.h"
+#include "Animation/SkeletalMeshActor.h"
+#include "Particles/ParticleSystemComponent.h"
+#include "Components/LightComponent.h"
+
+#include "UnrealEd.h"
+
+
+DECLARE_LOG_CATEGORY_EXTERN(LogMatineeToLevelSequence, Verbose, All);
diff --git a/Engine/Plugins/ScriptPlugin/Source/ScriptEditorPlugin/Private/ScriptFactory.cpp b/Engine/Plugins/ScriptPlugin/Source/ScriptEditorPlugin/Private/ScriptFactory.cpp
--- a/Engine/Plugins/ScriptPlugin/Source/ScriptEditorPlugin/Private/ScriptFactory.cpp
+++ b/Engine/Plugins/ScriptPlugin/Source/ScriptEditorPlugin/Private/ScriptFactory.cpp
@@ -145,7 +145,9 @@ EReimportResult::Type UReimportScriptFactory::Reimport(UObject* Obj)
 		return EReimportResult::Failed;
 	}
 
-	if (UFactory::StaticImportObject(ScriptClass->GetClass(), ScriptClass->GetOuter(), *ScriptClass->GetName(), RF_Public | RF_Standalone, *ResolvedSourceFilePath, NULL, this))
+	bool OutCanceled = false;
+
+	if (ImportObject(ScriptClass->GetClass(), ScriptClass->GetOuter(), *ScriptClass->GetName(), RF_Public | RF_Standalone, ResolvedSourceFilePath, nullptr, OutCanceled) != nullptr)
 	{
 		UE_LOG(LogScriptEditorPlugin, Log, TEXT("Imported successfully"));
 		// Try to find the outer package so we can dirty it up
@@ -158,6 +160,10 @@ EReimportResult::Type UReimportScriptFactory::Reimport(UObject* Obj)
 			ScriptClass->MarkPackageDirty();
 		}
 	}
+	else if(OutCanceled)
+	{
+		UE_LOG(LogScriptEditorPlugin, Warning, TEXT("-- import canceled"));
+	}
 	else
 	{
 		UE_LOG(LogScriptEditorPlugin, Warning, TEXT("-- import failed"));
diff --git a/Engine/Source/Developer/AssetTools/Private/AssetTools.cpp b/Engine/Source/Developer/AssetTools/Private/AssetTools.cpp
--- a/Engine/Source/Developer/AssetTools/Private/AssetTools.cpp
+++ b/Engine/Source/Developer/AssetTools/Private/AssetTools.cpp
@@ -33,6 +33,7 @@ FAssetTools::FAssetTools()
 	AllocatedCategoryBits.Add(TEXT("_BuiltIn_5"), FAdvancedAssetCategory(EAssetTypeCategories::UI, LOCTEXT("UserInterfaceAssetCategory", "User Interface")));
 	AllocatedCategoryBits.Add(TEXT("_BuiltIn_6"), FAdvancedAssetCategory(EAssetTypeCategories::Misc, LOCTEXT("MiscellaneousAssetCategory", "Miscellaneous")));
 	AllocatedCategoryBits.Add(TEXT("_BuiltIn_7"), FAdvancedAssetCategory(EAssetTypeCategories::Gameplay, LOCTEXT("GameplayAssetCategory", "Gameplay")));
+	AllocatedCategoryBits.Add(TEXT("_BuiltIn_8"), FAdvancedAssetCategory(EAssetTypeCategories::Media, LOCTEXT("MediaAssetCategory", "Media")));
 
 	EAssetTypeCategories::Type BlendablesCategoryBit = RegisterAdvancedAssetCategory(FName(TEXT("Blendables")), LOCTEXT("BlendablesAssetCategory", "Blendables"));
 
@@ -301,7 +302,7 @@ struct FRootedOnScope
 {
 	UObject* Obj;
 
-	FRootedOnScope(UObject* InObj) : Obj(NULL)
+	FRootedOnScope(UObject* InObj) : Obj(nullptr)
 	{
 		if (InObj && !InObj->IsRooted())
 		{
@@ -327,13 +328,13 @@ UObject* FAssetTools::CreateAsset(const FString& AssetName, const FString& Packa
 	if ( !ensure(AssetClass || Factory) )
 	{
 		FMessageDialog::Open( EAppMsgType::Ok, LOCTEXT("MustSupplyClassOrFactory", "The new asset wasn't created due to a problem finding the appropriate factory or class for the new asset.") );
-		return NULL;
+		return nullptr;
 	}
 
 	if ( AssetClass && Factory && !ensure(AssetClass->IsChildOf(Factory->GetSupportedClass())) )
 	{
 		FMessageDialog::Open( EAppMsgType::Ok, LOCTEXT("InvalidFactory", "The new asset wasn't created because the supplied factory does not support the supplied class.") );
-		return NULL;
+		return nullptr;
 	}
 
 	const FString PackageName = PackagePath + TEXT("/") + AssetName;
@@ -341,17 +342,17 @@ UObject* FAssetTools::CreateAsset(const FString& AssetName, const FString& Packa
 	// Make sure we can create the asset without conflicts
 	if ( !CanCreateAsset(AssetName, PackageName, LOCTEXT("CreateANewObject", "Create a new object")) )
 	{
-		return NULL;
+		return nullptr;
 	}
 
-	UClass* ClassToUse = AssetClass ? AssetClass : (Factory ? Factory->GetSupportedClass() : NULL);
+	UClass* ClassToUse = AssetClass ? AssetClass : (Factory ? Factory->GetSupportedClass() : nullptr);
 
-	UPackage* Pkg = CreatePackage(NULL,*PackageName);
-	UObject* NewObj = NULL;
+	UPackage* Pkg = CreatePackage(nullptr,*PackageName);
+	UObject* NewObj = nullptr;
 	EObjectFlags Flags = RF_Public|RF_Standalone;
 	if ( Factory )
 	{
-		NewObj = Factory->FactoryCreateNew(ClassToUse, Pkg, FName( *AssetName ), Flags, NULL, GWarn, CallingContext);
+		NewObj = Factory->FactoryCreateNew(ClassToUse, Pkg, FName( *AssetName ), Flags, nullptr, GWarn, CallingContext);
 	}
 	else if ( AssetClass )
 	{
@@ -454,7 +455,7 @@ UObject* FAssetTools::DuplicateAsset(const FString& AssetName, const FString& Pa
 	if ( !OriginalObject )
 	{
 		FMessageDialog::Open( EAppMsgType::Ok, LOCTEXT("InvalidSourceObject", "The new asset wasn't created due to a problem finding the object to duplicate.") );
-		return NULL;
+		return nullptr;
 	}
 
 	const FString PackageName = PackagePath + TEXT("/") + AssetName;
@@ -462,7 +463,7 @@ UObject* FAssetTools::DuplicateAsset(const FString& AssetName, const FString& Pa
 	// Make sure we can create the asset without conflicts
 	if ( !CanCreateAsset(AssetName, PackageName, LOCTEXT("DuplicateAnObject", "Duplicate an object")) )
 	{
-		return NULL;
+		return nullptr;
 	}
 
 	ObjectTools::FPackageGroupName PGN;
@@ -533,7 +534,7 @@ TArray<UObject*> FAssetTools::ImportAssets(const FString& DestinationPath)
 
 	if ( DesktopPlatform )
 	{
-		void* ParentWindowWindowHandle = NULL;
+		void* ParentWindowWindowHandle = nullptr;
 
 		IMainFrameModule& MainFrameModule = FModuleManager::LoadModuleChecked<IMainFrameModule>(TEXT("MainFrame"));
 		const TSharedPtr<SWindow>& MainFrameParentWindow = MainFrameModule.GetParentWindow();
@@ -558,7 +559,7 @@ TArray<UObject*> FAssetTools::ImportAssets(const FString& DestinationPath)
 	{
 		if ( OpenFilenames.Num() > 0 )
 		{
-			UFactory* ChosenFactory = NULL;
+			UFactory* ChosenFactory = nullptr;
 			if (FilterIndex > 0)
 			{
 				ChosenFactory = *FilterIndexToFactory.Find(FilterIndex);
@@ -626,94 +627,56 @@ TArray<UObject*> FAssetTools::ImportAssets(const TArray<FString>& Files, const F
 	TArray<TPair<FString, FString>> FilesAndDestinations;
 	ExpandDirectories(Files, RootDestinationPath, FilesAndDestinations);
 
-
-
 	SlowTask.EnterProgressFrame(1, LOCTEXT("Import_DeterminingImportTypes", "Determining asset types"));
 
-	if (SpecifiedFactory == NULL)
+	if (SpecifiedFactory == nullptr)
 	{	
 		// First instantiate one factory for each file extension encountered that supports the extension
+		// @todo import: gmp: show dialog in case of multiple matching factories
 		for( TObjectIterator<UClass> ClassIt; ClassIt; ++ClassIt )
 		{
-			if( (*ClassIt)->IsChildOf(UFactory::StaticClass()) && !((*ClassIt)->HasAnyClassFlags(CLASS_Abstract)) )
+			if (!(*ClassIt)->IsChildOf(UFactory::StaticClass()) || ((*ClassIt)->HasAnyClassFlags(CLASS_Abstract)))
 			{
-				UFactory* Factory = Cast<UFactory>( (*ClassIt)->GetDefaultObject() );
-				if( Factory->bEditorImport )
-				{
-					TArray<FString> FactoryExtensions;
-					Factory->GetSupportedFileExtensions(FactoryExtensions);
-					for ( auto& FileDest : FilesAndDestinations )
-					{
-						const FString FileExtension = FPaths::GetExtension(FileDest.Key);
-	
-						// Case insensitive string compare with supported formats of this factory
-						if ( FactoryExtensions.Contains(FileExtension) )
-						{
-							TArray<UFactory*>& ExistingFactories = ExtensionToFactoriesMap.FindOrAdd(FileExtension);
-							
-							// Do not remap extensions, just reuse the existing UFactory.
-							// There may be multiple UFactories, so we will keep track of all of them
-							bool bFactoryAlreadyInMap = false;
-							for ( auto FoundFactoryIt = ExistingFactories.CreateConstIterator(); FoundFactoryIt; ++FoundFactoryIt )
-							{
-								if ( (*FoundFactoryIt)->GetClass() == Factory->GetClass() )
-								{
-									bFactoryAlreadyInMap = true;
-									break;
-								}
-							}
-	
-							if ( !bFactoryAlreadyInMap )
-							{
-								// We found a factory for this file, it can be imported!
-								// Create a new factory of the same class and make sure it doesn't get GCed.
-								// The object will be removed from the root set at the end of this function.
-								UFactory* NewFactory = NewObject<UFactory>(GetTransientPackage(), Factory->GetClass());
-								if ( NewFactory->ConfigureProperties() )
-								{
-									NewFactory->AddToRoot();
-									ExistingFactories.Add(NewFactory);
-								}
-							}
-						}
-					}
-				}
+				continue;
 			}
-		}
-	}
-	else
-	{
-		if( SpecifiedFactory->bEditorImport )
-		{
-			TArray<FString> FactoryExtensions;
-			SpecifiedFactory->GetSupportedFileExtensions(FactoryExtensions);
-			for ( auto FileIt = Files.CreateConstIterator(); FileIt; ++FileIt )
+
+			UFactory* Factory = Cast<UFactory>( (*ClassIt)->GetDefaultObject() );
+
+			if (!Factory->bEditorImport)
 			{
-				const FString FileExtension = FPaths::GetExtension(*FileIt);
+				continue;
+			}
+
+			TArray<FString> FactoryExtensions;
+			Factory->GetSupportedFileExtensions(FactoryExtensions);
 
+			for (auto& FileDest : FilesAndDestinations)
+			{
+				const FString FileExtension = FPaths::GetExtension(FileDest.Key);
+	
 				// Case insensitive string compare with supported formats of this factory
 				if ( FactoryExtensions.Contains(FileExtension) )
 				{
 					TArray<UFactory*>& ExistingFactories = ExtensionToFactoriesMap.FindOrAdd(FileExtension);
-
+							
 					// Do not remap extensions, just reuse the existing UFactory.
 					// There may be multiple UFactories, so we will keep track of all of them
 					bool bFactoryAlreadyInMap = false;
 					for ( auto FoundFactoryIt = ExistingFactories.CreateConstIterator(); FoundFactoryIt; ++FoundFactoryIt )
 					{
-						if ( (*FoundFactoryIt)->GetClass() == SpecifiedFactory->GetClass() )
+						if ( (*FoundFactoryIt)->GetClass() == Factory->GetClass() )
 						{
 							bFactoryAlreadyInMap = true;
 							break;
 						}
 					}
-
+	
 					if ( !bFactoryAlreadyInMap )
 					{
 						// We found a factory for this file, it can be imported!
-						// Create a new factory of the same class and make sure it doesnt get GCed.
+						// Create a new factory of the same class and make sure it doesn't get GCed.
 						// The object will be removed from the root set at the end of this function.
-						UFactory* NewFactory = NewObject<UFactory>(GetTransientPackage(), SpecifiedFactory->GetClass());
+						UFactory* NewFactory = NewObject<UFactory>(GetTransientPackage(), Factory->GetClass());
 						if ( NewFactory->ConfigureProperties() )
 						{
 							NewFactory->AddToRoot();
@@ -724,6 +687,49 @@ TArray<UObject*> FAssetTools::ImportAssets(const TArray<FString>& Files, const F
 			}
 		}
 	}
+	else if( SpecifiedFactory->bEditorImport )
+	{
+		TArray<FString> FactoryExtensions;
+		SpecifiedFactory->GetSupportedFileExtensions(FactoryExtensions);
+
+		for (auto FileIt = Files.CreateConstIterator(); FileIt; ++FileIt)
+		{
+			const FString FileExtension = FPaths::GetExtension(*FileIt);
+
+			// Case insensitive string compare with supported formats of this factory
+			if (!FactoryExtensions.Contains(FileExtension))
+			{
+				continue;
+			}
+
+			TArray<UFactory*>& ExistingFactories = ExtensionToFactoriesMap.FindOrAdd(FileExtension);
+
+			// Do not remap extensions, just reuse the existing UFactory.
+			// There may be multiple UFactories, so we will keep track of all of them
+			bool bFactoryAlreadyInMap = false;
+			for ( auto FoundFactoryIt = ExistingFactories.CreateConstIterator(); FoundFactoryIt; ++FoundFactoryIt )
+			{
+				if ( (*FoundFactoryIt)->GetClass() == SpecifiedFactory->GetClass() )
+				{
+					bFactoryAlreadyInMap = true;
+					break;
+				}
+			}
+
+			if ( !bFactoryAlreadyInMap )
+			{
+				// We found a factory for this file, it can be imported!
+				// Create a new factory of the same class and make sure it doesnt get GCed.
+				// The object will be removed from the root set at the end of this function.
+				UFactory* NewFactory = NewObject<UFactory>(GetTransientPackage(), SpecifiedFactory->GetClass());
+				if ( NewFactory->ConfigureProperties() )
+				{
+					NewFactory->AddToRoot();
+					ExistingFactories.Add(NewFactory);
+				}
+			}
+		}
+	}
 
 	// We need to sort the factories so that they get tested in priority order
 	for (auto& ExtensionToFactories : ExtensionToFactoriesMap)
@@ -748,7 +754,7 @@ TArray<UObject*> FAssetTools::ImportAssets(const TArray<FString>& Files, const F
 		FString FileExtension = FPaths::GetExtension(Filename);
 
 		const TArray<UFactory*>* FactoriesPtr = ExtensionToFactoriesMap.Find(FileExtension);
-		UFactory* Factory = NULL;
+		UFactory* Factory = nullptr;
 		if ( FactoriesPtr )
 		{
 			const TArray<UFactory*>& Factories = *FactoriesPtr;
@@ -790,7 +796,7 @@ TArray<UObject*> FAssetTools::ImportAssets(const TArray<FString>& Files, const F
 			UE_LOG(LogAssetTools, Warning, TEXT("%s"), *Message.ToString() );
 		}
 
-		if ( Factory != NULL )
+		if ( Factory != nullptr )
 		{
 			UClass* ImportAssetType = Factory->SupportedClass;
 			bool bImportSucceeded = false;
@@ -810,7 +816,7 @@ TArray<UObject*> FAssetTools::ImportAssets(const TArray<FString>& Files, const F
 				continue;
 			}
 
-			UPackage* Pkg = CreatePackage(NULL, *PackageName);
+			UPackage* Pkg = CreatePackage(nullptr, *PackageName);
 			if ( !ensure(Pkg) )
 			{
 				// Failed to create the package to hold this asset for some reason
@@ -823,7 +829,7 @@ TArray<UObject*> FAssetTools::ImportAssets(const TArray<FString>& Files, const F
 
 			// Check for an existing object
 			UObject* ExistingObject = StaticFindObject( UObject::StaticClass(), Pkg, *Name );
-			if( ExistingObject != NULL )
+			if( ExistingObject != nullptr )
 			{
 				// If the existing object is one of the imports we've just created we can't replace or overwrite it
 				if (ReturnObjects.Contains(ExistingObject))
@@ -925,7 +931,7 @@ TArray<UObject*> FAssetTools::ImportAssets(const TArray<FString>& Files, const F
 							else
 							{
 								// succeed, recreate package since it has been deleted
-								Pkg = CreatePackage(NULL, *PackageName);
+								Pkg = CreatePackage(nullptr, *PackageName);
 							}
 						}
 						else
@@ -947,7 +953,7 @@ TArray<UObject*> FAssetTools::ImportAssets(const TArray<FString>& Files, const F
 			}
 
 			ImportAssetType = Factory->ResolveSupportedClass();
-			UObject* Result = UFactory::StaticImportObject( ImportAssetType, Pkg, FName( *Name ), RF_Public|RF_Standalone, bImportWasCancelled, *Filename, NULL, Factory );
+			UObject* Result = Factory->ImportObject(ImportAssetType, Pkg, FName(*Name), RF_Public | RF_Standalone, Filename, nullptr, bImportWasCancelled);
 
 			// Do not report any error if the operation was canceled.
 			if(!bImportWasCancelled)
@@ -1065,7 +1071,7 @@ void FAssetTools::CreateUniqueAssetName(const FString& InBasePackageName, const
 		}
 		else
 		{
-			bObjectExists = LoadObject<UObject>(NULL, *ObjectPath, NULL, LOAD_NoWarn | LOAD_NoRedirects) != NULL;
+			bObjectExists = LoadObject<UObject>(nullptr, *ObjectPath, nullptr, LOAD_NoWarn | LOAD_NoRedirects) != nullptr;
 		}
 		IntSuffix++;
 	}
@@ -1094,7 +1100,7 @@ bool FAssetTools::AssetUsesGenericThumbnail( const FAssetData& AssetData ) const
 		// It would be more correct here to find the rendering info for the generated class,
 		// but instead we are simply seeing if there is a thumbnail saved on disk for this asset
 		FString PackageFilename;
-		if ( FPackageName::DoesPackageExist(AssetData.PackageName.ToString(), NULL, &PackageFilename) )
+		if ( FPackageName::DoesPackageExist(AssetData.PackageName.ToString(), nullptr, &PackageFilename) )
 		{
 			TSet<FName> ObjectFullNames;
 			FThumbnailMap ThumbnailMap;
@@ -1118,15 +1124,15 @@ bool FAssetTools::AssetUsesGenericThumbnail( const FAssetData& AssetData ) const
 		// Unloaded non-blueprint asset. See if the class has a rendering info.
 		UClass* Class = FindObject<UClass>(ANY_PACKAGE, *AssetData.AssetClass.ToString());
 
-		UObject* ClassCDO = NULL;
-		if (Class != NULL)
+		UObject* ClassCDO = nullptr;
+		if (Class != nullptr)
 		{
 			ClassCDO = Class->GetDefaultObject();
 		}
 
 		// Get the rendering info for this object
-		FThumbnailRenderingInfo* RenderInfo = NULL;
-		if (ClassCDO != NULL)
+		FThumbnailRenderingInfo* RenderInfo = nullptr;
+		if (ClassCDO != nullptr)
 		{
 			RenderInfo = GUnrealEd->GetThumbnailManager()->GetRenderingInfo( ClassCDO );
 		}
@@ -1151,11 +1157,11 @@ void FAssetTools::DiffAgainstDepot( UObject* InObject, const FString& InPackageP
 	FSourceControlStatePtr SourceControlState = SourceControlProvider.GetState(SourceControlHelpers::PackageFilename(InPackagePath), EStateCacheUsage::Use);
 
 	// If we have an asset and its in SCC..
-	if( SourceControlState.IsValid() && InObject != NULL && SourceControlState->IsSourceControlled() )
+	if( SourceControlState.IsValid() && InObject != nullptr && SourceControlState->IsSourceControlled() )
 	{
 		// Get the file name of package
 		FString RelativeFileName;
-		if(FPackageName::DoesPackageExist(InPackagePath, NULL, &RelativeFileName))
+		if(FPackageName::DoesPackageExist(InPackagePath, nullptr, &RelativeFileName))
 		{
 			if(SourceControlState->GetHistorySize() > 0)
 			{
@@ -1171,12 +1177,12 @@ void FAssetTools::DiffAgainstDepot( UObject* InObject, const FString& InPackageP
 					TGuardValue<bool> DisableCompileOnLoad(GForceDisableBlueprintCompileOnLoad, true);
 
 					// Try and load that package
-					UPackage* TempPackage = LoadPackage(NULL, *TempFileName, LOAD_ForDiff);
-					if(TempPackage != NULL)
+					UPackage* TempPackage = LoadPackage(nullptr, *TempFileName, LOAD_ForDiff);
+					if(TempPackage != nullptr)
 					{
 						// Grab the old asset from that old package
 						UObject* OldObject = FindObject<UObject>(TempPackage, *InPackageName);
-						if(OldObject != NULL)
+						if(OldObject != nullptr)
 						{
 							/* Set the revision information*/
 							FRevisionInfo OldRevision;
@@ -1197,9 +1203,9 @@ void FAssetTools::DiffAgainstDepot( UObject* InObject, const FString& InPackageP
 
 void FAssetTools::DiffAssets(UObject* OldAsset, UObject* NewAsset, const struct FRevisionInfo& OldRevision, const struct FRevisionInfo& NewRevision) const
 {
-	if(OldAsset == NULL || NewAsset == NULL)
+	if(OldAsset == nullptr || NewAsset == nullptr)
 	{
-		UE_LOG(LogAssetTools, Warning, TEXT("DiffAssets: One of the supplied assets was NULL."));
+		UE_LOG(LogAssetTools, Warning, TEXT("DiffAssets: One of the supplied assets was nullptr."));
 		return;
 	}
 
@@ -1234,7 +1240,7 @@ FString FAssetTools::DumpAssetToTempFile(UObject* Asset) const
 	const FExportObjectInnerContext Context;
 
 	// Export asset to archive
-	UExporter::ExportToOutputDevice(&Context, Asset, NULL, Archive, TEXT("copy"), 0, PPF_ExportsNotFullyQualified|PPF_Copy|PPF_Delimited, false, Asset->GetOuter());
+	UExporter::ExportToOutputDevice(&Context, Asset, nullptr, Archive, TEXT("copy"), 0, PPF_ExportsNotFullyQualified|PPF_Copy|PPF_Delimited, false, Asset->GetOuter());
 
 	// Used to generate unique file names during a run
 	static int TempFileNum = 0;
@@ -1274,7 +1280,7 @@ bool FAssetTools::CreateDiffProcess(const FString& DiffCommand,  const FString&
 	while (bTryRunDiff)
 	{
 		// Fire process
-		if (FPlatformProcess::CreateProc(*NewDiffCommand, *Arguments, true, false, false, NULL, 0, NULL, NULL).IsValid())
+		if (FPlatformProcess::CreateProc(*NewDiffCommand, *Arguments, true, false, false, nullptr, 0, nullptr, nullptr).IsValid())
 		{
 			return true;
 		}
@@ -1349,7 +1355,7 @@ void FAssetTools::MigratePackages(const TArray<FName>& PackageNamesToMigrate) co
 void FAssetTools::OnNewImportRecord(UClass* AssetType, const FString& FileExtension, bool bSucceeded, bool bWasCancelled, const FDateTime& StartTime)
 {
 	// Don't attempt to report usage stats if analytics isn't available
-	if(AssetType != NULL && FEngineAnalytics::IsAvailable())
+	if(AssetType != nullptr && FEngineAnalytics::IsAvailable())
 	{
 		TArray<FAnalyticsEventAttribute> Attribs;
 		Attribs.Add(FAnalyticsEventAttribute(TEXT("AssetType"), AssetType->GetName()));
@@ -1365,7 +1371,7 @@ void FAssetTools::OnNewImportRecord(UClass* AssetType, const FString& FileExtens
 void FAssetTools::OnNewCreateRecord(UClass* AssetType, bool bDuplicated)
 {
 	// Don't attempt to report usage stats if analytics isn't available
-	if(AssetType != NULL && FEngineAnalytics::IsAvailable())
+	if(AssetType != nullptr && FEngineAnalytics::IsAvailable())
 	{
 		TArray<FAnalyticsEventAttribute> Attribs;
 		Attribs.Add(FAnalyticsEventAttribute(TEXT("AssetType"), AssetType->GetName()));
@@ -1449,7 +1455,7 @@ bool FAssetTools::CanCreateAsset(const FString& AssetName, const FString& Packag
 	}
 
 	// Find (or create!) the desired package for this object
-	UPackage* Pkg = CreatePackage(NULL,*PackageName);
+	UPackage* Pkg = CreatePackage(nullptr,*PackageName);
 
 	// Handle fully loading packages before creating new objects.
 	TArray<UPackage*> TopLevelPackages;
@@ -1470,7 +1476,7 @@ bool FAssetTools::CanCreateAsset(const FString& AssetName, const FString& Packag
 
 	// Check for an existing object
 	UObject* ExistingObject = StaticFindObject( UObject::StaticClass(), Pkg, *AssetName );
-	if( ExistingObject != NULL )
+	if( ExistingObject != nullptr )
 	{
 		// Object already exists in either the specified package or another package.  Check to see if the user wants
 		// to replace the object.
@@ -1493,7 +1499,7 @@ bool FAssetTools::CanCreateAsset(const FString& AssetName, const FString& Packag
 				CollectGarbage( GARBAGE_COLLECTION_KEEPFLAGS );
 
 				// Old package will be GC'ed... create a new one here
-				Pkg = CreatePackage(NULL,*PackageName);
+				Pkg = CreatePackage(nullptr,*PackageName);
 			}
 			else
 			{
@@ -1570,7 +1576,7 @@ void FAssetTools::MigratePackages_ReportConfirmed(TArray<FString> ConfirmedPacka
 	FString DestinationFolder;
 	if ( ensure(DesktopPlatform) )
 	{
-		void* ParentWindowWindowHandle = NULL;
+		void* ParentWindowWindowHandle = nullptr;
 
 		IMainFrameModule& MainFrameModule = FModuleManager::LoadModuleChecked<IMainFrameModule>(TEXT("MainFrame"));
 		const TSharedPtr<SWindow>& MainFrameParentWindow = MainFrameModule.GetParentWindow();
@@ -1644,7 +1650,7 @@ void FAssetTools::MigratePackages_ReportConfirmed(TArray<FString> ConfirmedPacka
 
 			const FString& PackageName = *PackageNameIt;
 			FString SrcFilename;
-			if ( !FPackageName::DoesPackageExist(PackageName, NULL, &SrcFilename) )
+			if ( !FPackageName::DoesPackageExist(PackageName, nullptr, &SrcFilename) )
 			{
 				const FText ErrorMessage = FText::Format(LOCTEXT("MigratePackages_PackageMissing", "{0} does not exist on disk."), FText::FromString(PackageName));
 				UE_LOG(LogAssetTools, Warning, TEXT("%s"), *ErrorMessage.ToString());
diff --git a/Engine/Source/Developer/AssetTools/Public/AssetTypeCategories.h b/Engine/Source/Developer/AssetTools/Public/AssetTypeCategories.h
--- a/Engine/Source/Developer/AssetTools/Public/AssetTypeCategories.h
+++ b/Engine/Source/Developer/AssetTools/Public/AssetTypeCategories.h
@@ -22,9 +22,10 @@ namespace EAssetTypeCategories
 		Misc = 1 << 6,
 		Gameplay = 1 << 7,
 		Blueprint = 1 << 8,
+		Media = 1 << 9,
 
 		// Items below this will be allocated at runtime via RegisterAdvancedAssetCategory
-		FirstUser = 1 << 9,
+		FirstUser = 1 << 10,
 		LastUser = 1 << 31,
 		// Last allowed value is 1 << 31
 	};
diff --git a/Engine/Source/Developer/Profiler/Private/ProfilerManager.cpp b/Engine/Source/Developer/Profiler/Private/ProfilerManager.cpp
--- a/Engine/Source/Developer/Profiler/Private/ProfilerManager.cpp
+++ b/Engine/Source/Developer/Profiler/Private/ProfilerManager.cpp
@@ -270,7 +270,7 @@ bool FProfilerManager::Tick( float DeltaTime )
 		const double SessionMemory = MBInv*ProfilerSession->GetMemoryUsage();
 		const double PhysMemory = MBInv*DiffPhys;
 
-		UE_LOG( LogStats, Log, TEXT( "ProfilerSession: %6.2f MB (%6.2f MB) # (%6.2f MB) / %7u -> %4u" ), 
+		UE_LOG( LogStats, VeryVerbose, TEXT( "ProfilerSession: %6.2f MB (%6.2f MB) # (%6.2f MB) / %7u -> %4u" ), 
 			SessionMemory,
 			PhysMemory,
 			PhysMemory - SessionMemory,
diff --git a/Engine/Source/Editor/ContentBrowser/Private/AssetContextMenu.cpp b/Engine/Source/Editor/ContentBrowser/Private/AssetContextMenu.cpp
--- a/Engine/Source/Editor/ContentBrowser/Private/AssetContextMenu.cpp
+++ b/Engine/Source/Editor/ContentBrowser/Private/AssetContextMenu.cpp
@@ -1816,7 +1816,16 @@ bool FAssetContextMenu::CanExecuteSyncToAssetTree() const
 
 bool FAssetContextMenu::CanExecuteFindInExplorer() const
 {
-	return SelectedAssets.Num() > 0;
+	// selection must contain at least one asset that has already been saved to disk
+	for (const FAssetData& Asset : SelectedAssets)
+	{
+		if ((Asset.PackageFlags & PKG_NewlyCreated) == 0)
+		{
+			return true;
+		}
+	}
+
+	return false;
 }
 
 bool FAssetContextMenu::CanExecuteCreateBlueprintUsing() const
diff --git a/Engine/Source/Editor/ContentBrowser/Private/ContentBrowserUtils.cpp b/Engine/Source/Editor/ContentBrowser/Private/ContentBrowserUtils.cpp
--- a/Engine/Source/Editor/ContentBrowser/Private/ContentBrowserUtils.cpp
+++ b/Engine/Source/Editor/ContentBrowser/Private/ContentBrowserUtils.cpp
@@ -1522,7 +1522,7 @@ FText ContentBrowserUtils::GetExploreFolderText()
 {
 	FFormatNamedArguments Args;
 	Args.Add(TEXT("FileManagerName"), FPlatformMisc::GetFileManagerName());
-	return FText::Format(NSLOCTEXT("GenericPlatform", "ShowInFileManager", "Show In {FileManagerName}"), Args);
+	return FText::Format(NSLOCTEXT("GenericPlatform", "ShowInFileManager", "Show in {FileManagerName}"), Args);
 }
 
 static const auto CVarMaxFullPathLength = 
diff --git a/Engine/Source/Editor/ContentBrowser/Private/NewAssetOrClassContextMenu.cpp b/Engine/Source/Editor/ContentBrowser/Private/NewAssetOrClassContextMenu.cpp
--- a/Engine/Source/Editor/ContentBrowser/Private/NewAssetOrClassContextMenu.cpp
+++ b/Engine/Source/Editor/ContentBrowser/Private/NewAssetOrClassContextMenu.cpp
@@ -334,6 +334,9 @@ void FNewAssetOrClassContextMenu::MakeContextMenu(
 
 			TArray<FAdvancedAssetCategory> AdvancedAssetCategories;
 			AssetToolsModule.Get().GetAllAdvancedAssetCategories(/*out*/ AdvancedAssetCategories);
+			AdvancedAssetCategories.Sort([](const FAdvancedAssetCategory& A, const FAdvancedAssetCategory& B) {
+				return (A.CategoryName.CompareToCaseIgnored(B.CategoryName) < 0);
+			});
 
 			for (const FAdvancedAssetCategory& AdvancedAssetCategory : AdvancedAssetCategories)
 			{
@@ -349,11 +352,11 @@ void FNewAssetOrClassContextMenu::MakeContextMenu(
 							FirstSelectedPath, 
 							InOnNewAssetRequested, 
 							FCanExecuteAction() // We handle this at this level, rather than at the sub-menu item level
-							),
+						),
 						FUIAction(
 							FExecuteAction(),
 							CanExecuteAssetActionsDelegate
-							),
+						),
 						NAME_None,
 						EUserInterfaceActionType::Button
 						);
diff --git a/Engine/Source/Editor/CurveAssetEditor/Private/CurveAssetEditor.cpp b/Engine/Source/Editor/CurveAssetEditor/Private/CurveAssetEditor.cpp
--- a/Engine/Source/Editor/CurveAssetEditor/Private/CurveAssetEditor.cpp
+++ b/Engine/Source/Editor/CurveAssetEditor/Private/CurveAssetEditor.cpp
@@ -209,8 +209,9 @@ TSharedPtr<FExtender> FCurveAssetEditor::GetToolbarExtender()
 
 			ToolbarBuilder.BeginSection("Snap");
 			{
-				ToolbarBuilder.AddToolBarButton(FRichCurveEditorCommands::Get().ToggleSnapping);
+				ToolbarBuilder.AddToolBarButton(FRichCurveEditorCommands::Get().ToggleInputSnapping);
 				ToolbarBuilder.AddWidget(InputSnapWidget);
+				ToolbarBuilder.AddToolBarButton(FRichCurveEditorCommands::Get().ToggleOutputSnapping);
 				ToolbarBuilder.AddWidget(OutputSnapWidget);
 			}
 			ToolbarBuilder.EndSection();
