Copying //UE4/Dev-Editor to //UE4/Dev-Main (Source: //UE4/Dev-Editor @ 3152045)
==========================
MAJOR FEATURES + CHANGES
==========================

Change 3106449 on 2016/08/30 by [REDACTED_PERSON]

Change 3133997 on 2016/09/21 by [REDACTED_PERSON]

Change 3134132 on 2016/09/21 by [REDACTED_PERSON]

	Added the "unattended" flag when running the localzation commandlets via UAT

Change 3134147 on 2016/09/21 by [REDACTED_PERSON]

	Core - Add multitouch support to windows.

Change 3134349 on 2016/09/21 by [REDACTED_PERSON]

Change 3134355 on 2016/09/21 by [REDACTED_PERSON]

Change 3134447 on 2016/09/21 by [REDACTED_PERSON]

	The SpriteComponent is now properly reused when duplicating a light.

Change 3134451 on 2016/09/21 by [REDACTED_PERSON]

	Make sure when we re-import we dont try to match the mesh name if the option bCombineToSingle is true

Change 3134457 on 2016/09/21 by [REDACTED_PERSON]

	FInternationalization::Leetify now deals with escape sequences correctly.

	This also changes it to mark all leetified text with the start and end marker, and the unit tests have been updated accordingly.

Change 3134685 on 2016/09/21 by [REDACTED_PERSON]

	Cleaned up some of the UI for the new material slot workflow (part 1)

Change 3134743 on 2016/09/21 by [REDACTED_PERSON]

	Fix crash clicking OK to the import mesh dialog in persona when a path is not selected.  The OK button is now disabled until a valid path is selected

Change 3134760 on 2016/09/21 by [REDACTED_PERSON]

	Guard against cascade crash if cascade was shutdown but there is an old cascade particle system component in the undo buffer that attempts to refresh cascade

Change 3134837 on 2016/09/21 by [REDACTED_PERSON]

	Fix static analysis warnings in this file

Change 3134939 on 2016/09/21 by [REDACTED_PERSON]

	Automation - Moving the functional testing hooks for the editor into a plugin, the goal is to make it so you don't have to see all the extra stuff automation adds to the editor until the game project needs it.  Refactoring the way automation messages are serialized, no longer relying on a comma delimited string to serialize automation data.  Improving the filter for automation window so it lets you put in a path, and find it that way.  Fixing RunTests from the commandline so that it properly finds the tests, instead of jumping out of the loop as soon as it finds any one match.

Change 3134941 on 2016/09/21 by [REDACTED_PERSON]

	Automation - Functional tests now support warnings as errors on a per test basis.  Continued itteration on screenshot comparison.

Change 3135051 on 2016/09/21 by [REDACTED_PERSON]

	Fixed size of material slots in persona details panel

Change 3135076 on 2016/09/21 by [REDACTED_PERSON]

	Automation - Adding a missing file.

Change 3135584 on 2016/09/22 by [REDACTED_PERSON]

	Removed unused ULandscapeInfo::BeginDestroy
	Slight cleanup to FLandscapeEditDataInterface::DeleteLayer

Change 3135585 on 2016/09/22 by [REDACTED_PERSON]

	New Count algo
	Const'd Copy/Accumulate/Transform

Change 3135599 on 2016/09/22 by [REDACTED_PERSON]

	Resolve crash if a Landscape Layer Info object is force deleted while it is in use :(

Change 3135724 on 2016/09/22 by [REDACTED_PERSON]

Change 3135726 on 2016/09/22 by [REDACTED_PERSON]

	Slate - Removing the SLATE_PRE_MULTIPLY macro.

Change 3135730 on 2016/09/22 by [REDACTED_PERSON]

	UMG - Exposing the 2D hit location that the WidgetInteractionComponent hit on the WidgetComponent.

Change 3135738 on 2016/09/22 by [REDACTED_PERSON]

	Ensure any handles to the backbuffer are released before the backbuffer is resized

	https://jira.it.epicgames.net/browse/UE-30488

Change 3135810 on 2016/09/22 by [REDACTED_PERSON]

	Build: Compile tools before running automated tests.

Change 3135993 on 2016/09/22 by [REDACTED_PERSON]

	Ensure you can unpause after toggling play/pause with a keyboard shortcut.  We were processing the keybinding on key down AND up which caused the unpause to instantly pause again
	https://jira.it.epicgames.net/browse/UE-36276

Change 3136257 on 2016/09/22 by [REDACTED_PERSON]

	Fixed assign of materials to components not working if you multi-select multiple components of a blueprint and assign the material.  The construction script was running before all the material had been set on all components.

Change 3136318 on 2016/09/22 by [REDACTED_PERSON]

Change 3136494 on 2016/09/22 by [REDACTED_PERSON]

	Fix crash restarting a mission in Odin

Change 3136741 on 2016/09/22 by [REDACTED_PERSON]

	Fixed SButton and SWidget to not multiply color and opacity, since that's already being done by SCompoundWidget

Change 3137711 on 2016/09/23 by [REDACTED_PERSON]

	Added guard against slate rendering with deleted materials.  In this will now be caught with a looged with the deleted material name.

Change 3137713 on 2016/09/23 by [REDACTED_PERSON]

	Removed the pooled draw elements stuff which is not used

Change 3137791 on 2016/09/23 by [REDACTED_PERSON]

	MediaPlayer - Removing .png from one of the style files, as that's not required.  Was causing warnings to be logged on load.

Change 3137793 on 2016/09/23 by [REDACTED_PERSON]

	Localization - Adjusting some log statements to say with instead of w/

Change 3137796 on 2016/09/23 by [REDACTED_PERSON]

	Slate - Adding missing and replacing corrupted style files in the editor.

Change 3137864 on 2016/09/23 by [REDACTED_PERSON]

	Fixed "actors are referenced are you sure you want to delete" dialog appearing when you are copying a lot of actors to another level

Change 3137876 on 2016/09/23 by [REDACTED_PERSON]

	Added allocation tagging to MProf2

	This hooks into the existing FScopeCycleCounterUObject used by the stats system to track object tags (object, package, and class) when an allocation is made. Tags should be in the format "Category:Tag", and the "Tag" part may include "/" to create levels in the tag hierarchy (as shown in the MProf2 tool).

Change 3137982 on 2016/09/23 by [REDACTED_PERSON]

	Added Invoke() support to CopyIf

Change 3137983 on 2016/09/23 by [REDACTED_PERSON]

	Added ULandscapeInfo::ForAllLandscapeProxies to clean up a lot of horrible code
	- Removed some dependencies on ALandscape and ALandscapeStreamingProxy in the process :)
	Also made FLandscapeEditorLayerSettings's constructor explicit to clean up some other horrible code

Change 3138053 on 2016/09/23 by [REDACTED_PERSON]

	Fixed Child Actor Template properties not visible after changing Child Actor Class

Change 3138079 on 2016/09/23 by [REDACTED_PERSON]

	Fixing some Clang warnings

Change 3138087 on 2016/09/23 by [REDACTED_PERSON]

	Added GetResourceSizeEx and GetResourceSizeBytes, and deprecated GetResourceSize

	GetResourceSizeEx populates a struct which reports not only how much memory is allocated, but also which arenas the memory was allocated from. GetResourceSizeBytes just wraps a call to GetResourceSizeEx and returns the total size from all arenas (this has the same behavior as the now deprecated GetResourceSize).

	Classes that used to override GetResourceSize should instead override GetResourceSizeEx and report their allocations as appropriate.

Change 3138127 on 2016/09/23 by [REDACTED_PERSON]

	Fixed crash when merging levels containing landscape

Change 3138821 on 2016/09/23 by [REDACTED_PERSON]

	Fixes "Select all input node" doesn't work properly on output node.

Change 3138915 on 2016/09/23 by [REDACTED_PERSON]

	Disable "select all linked nodes" for output nodes in material editor

Change 3139341 on 2016/09/25 by [REDACTED_PERSON]

	Automation - Moving the Blueprint Compiler Tests into the RuntimeTests plugin, and making the context client only since that's the only valid place to runt these tests.

Change 3139342 on 2016/09/25 by [REDACTED_PERSON]

	Landscape - Fixing some compiler errors on mac.

Change 3139345 on 2016/09/25 by [REDACTED_PERSON]

	Automation - Spelling and cleanup.

Change 3139346 on 2016/09/25 by [REDACTED_PERSON]

	Engine - Changing a check to an ensure, there's no reason to crash if this happens.

Change 3139347 on 2016/09/25 by [REDACTED_PERSON]

	Automation  - Making EFunctionalTestResult a BlueprintType

Change 3139348 on 2016/09/25 by [REDACTED_PERSON]

	Automation - Adding another test map.

Change 3139676 on 2016/09/26 by [REDACTED_PERSON]

Change 3139702 on 2016/09/26 by [REDACTED_PERSON]

	UMG - Making GetLocalHitLocation on UWidgetComponent virtual.

Change 3139760 on 2016/09/26 by [REDACTED_PERSON]

	Make sure we remove override materials from the list when the mesh point by the component has less materials.

Change 3139761 on 2016/09/26 by [REDACTED_PERSON]

	Added additional validation code to FPropertyNode to now properly validate TMap value and key nodes.

Change 3139843 on 2016/09/26 by [REDACTED_PERSON]

Change 3139880 on 2016/09/26 by [REDACTED_PERSON]

	QAGame - Update multitouch test with less crashy assets?

Change 3139908 on 2016/09/26 by [REDACTED_PERSON]

	Fix selections having the potential to be out of sync after undo/redo

Change 3139928 on 2016/09/26 by [REDACTED_PERSON]

	Automation - tweaking the test maps some more.

Change 3140646 on 2016/09/26 by [REDACTED_PERSON]

	Fix false positive with the test for vaild materials being rendered by slate.

Change 3140912 on 2016/09/26 by [REDACTED_PERSON]

	Core - Fix multitouch ifdef which was preventing it from actually being enabled.Also update test level blueprint so that it's actually testable.

Change 3141218 on 2016/09/27 by [REDACTED_PERSON]

	PR #2798: BP open anim interferes with mouse movement (Contributed by [REDACTED_PERSON])

Change 3141223 on 2016/09/27 by [REDACTED_PERSON]

	Updated UTextProperty::Identical to no longer compare display strings all the time

	It was supposed to compare the identity at runtime as the display string can change at runtime. This was preventing FText properties from being used in TSet/TMap as the hash needs to be consistent.

Change 3141242 on 2016/09/27 by [REDACTED_PERSON]

	Fixed various issues where making changes to components in the Blueprint Editor could cause a crash if there is an active component in the component visualizer.
	Duplicated CL 3139370 from //UE4/Release-4.13
	Duplicated CL 3139878 from //UE4/Release-4.13

Change 3141323 on 2016/09/27 by [REDACTED_PERSON]

Change 3141798 on 2016/09/27 by tim.gautier

	Added UMG_DisplayWidget

Change 3143038 on 2016/09/28 by [REDACTED_PERSON]

	Added extra context to FTableRowBase::OnPostDataImport

	It now takes the owning data table and the row name as parameters. This allows is to do more useful fix-up that depends on the context of where it's used (such as stabilized text keys).

Change 3143039 on 2016/09/28 by [REDACTED_PERSON]

	Optimized UTextProperty::Identical to use a pointer comparison rather than read out the identity

	Only texts which have the same display string pointer can have the same identity, so this is a much faster check.

Change 3143098 on 2016/09/28 by [REDACTED_PERSON]

	Fixed crash when loading duplicated landscape levels

Change 3143300 on 2016/09/28 by [REDACTED_PERSON]

	Fixed crash when duplicating a level containing a landscape through the content browser

Change 3143389 on 2016/09/28 by [REDACTED_PERSON]

	LastResort is no longer staged in shipping builds

	You can now use the Content/SlateDebug folder to store any Slate resources that shouldn't be used in a shipping build (either for the Engine or for a game).

	This also removes the old bUsesSlateEditorStyle hack as everything should have migrated away from that now.

Change 3143565 on 2016/09/28 by [REDACTED_PERSON]

	Fixed this file

Change 3143717 on 2016/09/28 by [REDACTED_PERSON]

	Fixed lowercase

Change 3143798 on 2016/09/28 by [REDACTED_PERSON]

	Fixed StreamingPauseRendering code to not attempt to tick the scene viewport on the slate loading thread while where are blocked on level streaming.  The viewport is rendered once and the render target is passed to slate instead of the entire viewport

Change 3143820 on 2016/09/28 by [REDACTED_PERSON]

	Use the PersonaToolikit to get the mesh

Change 3143833 on 2016/09/28 by [REDACTED_PERSON]

	Added guard for UE-36499

Change 3144144 on 2016/09/28 by [REDACTED_PERSON]

	Fix constructor init order

Change 3144821 on 2016/09/29 by [REDACTED_PERSON]

	Fixed feedback loop in FSlateEditableTextLayout::ComputeDesiredSize

	If the wrapping width is less than the scrollbar width, the scrollbar could constantly re-appear and then disappear.

Change 3144867 on 2016/09/29 by [REDACTED_PERSON]

	Added a commandlet to allow command line importing of any asset type the engine supports.
	- Intermediate checkin

Change 3144875 on 2016/09/29 by [REDACTED_PERSON]

	Automation - Adding CornellBox example map, fixing some additional tests, removing the Movement Test.

Change 3144975 on 2016/09/29 by [REDACTED_PERSON]

	Disable auto-applying of scalability settings
	The user will still be asked if the settings should be applied but if they do nothing the notification will just go away without setting stuff.

Change 3145274 on 2016/09/29 by [REDACTED_PERSON]

	New asset menu no longer has a scrollbar on most displays

Change 3146004 on 2016/09/29 by [REDACTED_PERSON]

	Adding two FBX test files to help facilitate QA coverage on the new material/section workflow.

Change 3146377 on 2016/09/30 by [REDACTED_PERSON]

	Fixed landscape rendering errors after using the "change component size" tool
	- also affected initial import of a landscape+weightmaps

Change 3146455 on 2016/09/30 by [REDACTED_PERSON]

	Fixing more menus that had scrollbars due to only being able to use half the vertical resolution

Change 3146466 on 2016/09/30 by [REDACTED_PERSON]

	Fixed not being able to erase foliage attached to BSP

Change 3146471 on 2016/09/30 by [REDACTED_PERSON]

	Can no longer localize sequences via the Content Browser

Change 3146569 on 2016/09/30 by [REDACTED_PERSON]

	Fixed UGatherTextFromSourceCommandlet::ParseSourceText being able to underflow while parsing

Change 3147116 on 2016/09/30 by [REDACTED_PERSON]

Change 3148091 on 2016/10/01 by [REDACTED_PERSON]

	Adding a map and relevant assets to facilitate testing around the new material blueprint nodes - GetMaterialSlotNames, GetMaterialIndex, and SetMaterialByName, respectively

Change 3148714 on 2016/10/03 by [REDACTED_PERSON]

	PR #2770: [Git plugin] Fix bug where history and merging do not work if the user has format.pretty settings in their gitconfig (Contributed by [REDACTED_PERSON])

Change 3148793 on 2016/10/03 by [REDACTED_PERSON]

	Automation - TPS for SM_Cornellbox

Change 3148801 on 2016/10/03 by [REDACTED_PERSON]

	PR #2820: [WidgetBlueprintLibrary] Throwing Essential-Functionality-Understanding PIE Errors when Player Controller ptr not supplied (Contributed by [REDACTED_PERSON])

Change 3148805 on 2016/10/03 by [REDACTED_PERSON]

	PR #2822: Add missing base includes and forward declarations for UProgressBar and UTextBlock (Contributed by [REDACTED_PERSON])

Change 3148813 on 2016/10/03 by [REDACTED_PERSON]

	UMG - The retainer now contains the SVirtualWindow directly in the hierarchy.  This should now make it possible to reliably focus elements inside the retainer widget.

Change 3148855 on 2016/10/03 by [REDACTED_PERSON]

	Fixed performance regression when importing landscape heightmaps

Change 3149482 on 2016/10/03 by [REDACTED_PERSON]

	Added link to Support landing page to Help menu

Change 3149520 on 2016/10/03 by tim.gautier

	Edited UMG_Behavior - Collapsed now has a toggleable state to demonstrate functionality.

Change 3149945 on 2016/10/04 by [REDACTED_PERSON]

	Fixed invisible landscape components when using tessellation on landscape material

Change 3149951 on 2016/10/04 by [REDACTED_PERSON]

	Reduced material update log spam when creating and editing landscapes
	- Some will remain, landscape fundementally manipulates a lot of material instances

Change 3150143 on 2016/10/04 by [REDACTED_PERSON]

	Fix not crashing when trying to clear objects which have a specific class filter in the property settings
	https://jira.it.epicgames.net/browse/UE-36692

Change 3150614 on 2016/10/04 by [REDACTED_PERSON]

	Git - Disabling some logging in non-debug builds of the git source control plugin.

Change 3151647 on 2016/10/05 by [REDACTED_PERSON]

	Fix loc warnings in these files (duplicate loc keys)

Change 3151679 on 2016/10/05 by [REDACTED_PERSON]

	Editor - Fixing the build, removing the VREditor module from the dynamically loaded list.

Change 3151722 on 2016/10/05 by [REDACTED_PERSON]

	Fix breakage to LandscapeEdModeSplineTools.cpp caused by Dev-VREditor from main integration

Change 3151816 on 2016/10/05 by [REDACTED_PERSON]

	Fixed more breakage to Landscape caused by Dev-VREditor from main integration

diff --git a/Engine/Config/BaseEditorPerProjectUserSettings.ini b/Engine/Config/BaseEditorPerProjectUserSettings.ini
--- a/Engine/Config/BaseEditorPerProjectUserSettings.ini
+++ b/Engine/Config/BaseEditorPerProjectUserSettings.ini
@@ -64,6 +64,7 @@ bDisplayUIExtensionPoints=False
 bUseCurvesForDistributions=False
 bAutoloadCheckedOutPackages=False
 bAutomaticallyHotReloadNewClasses=True
+bDisplayEngineVersionInBadge=True
 
 [/Script/EditorStyle.EditorStyleSettings]
 ; Applies a color vision deficiency filter to the entire editor
diff --git a/Engine/Plugins/Developer/GitSourceControl/Source/GitSourceControl/Private/GitSourceControlUtils.cpp b/Engine/Plugins/Developer/GitSourceControl/Source/GitSourceControl/Private/GitSourceControlUtils.cpp
--- a/Engine/Plugins/Developer/GitSourceControl/Source/GitSourceControl/Private/GitSourceControlUtils.cpp
+++ b/Engine/Plugins/Developer/GitSourceControl/Source/GitSourceControl/Private/GitSourceControlUtils.cpp
@@ -94,18 +94,20 @@ static bool RunCommandInternalRaw(const FString& InCommand, const FString& InPat
 
 	FullCommand += LogableCommand;
 
-	// @todo: temporary debug logs
-	//UE_LOG(LogSourceControl, Log, TEXT("RunCommandInternalRaw: 'git %s'"), *FullCommand);
+#if UE_BUILD_DEBUG
+	UE_LOG(LogSourceControl, Log, TEXT("RunCommandInternalRaw: 'git %s'"), *FullCommand);
+#endif
 	
 	FPlatformProcess::ExecProcess(*InPathToGitBinary, *FullCommand, &ReturnCode, &OutResults, &OutErrors);
 	
-	//// @todo: temporary debug logs
-	//UE_LOG(LogSourceControl, Log, TEXT("RunCommandInternalRaw: 'OutResults=\n%s'"), *OutResults);
-	//if (ReturnCode != 0)
-	//{
-	//	// @todo: temporary debug logs
-	//	UE_LOG(LogSourceControl, Warning, TEXT("RunCommandInternalRaw: 'OutErrors=\n%s'"), *OutErrors);
-	//}
+#if UE_BUILD_DEBUG
+	UE_LOG(LogSourceControl, Log, TEXT("RunCommandInternalRaw: 'OutResults=\n%s'"), *OutResults);
+
+	if (ReturnCode != 0)
+	{
+		UE_LOG(LogSourceControl, Warning, TEXT("RunCommandInternalRaw: 'OutErrors=\n%s'"), *OutErrors);
+	}
+#endif
 
 	return ReturnCode == 0;
 }
@@ -866,6 +868,7 @@ bool RunGetHistory(const FString& InPathToGitBinary, const FString& InRepository
 		Parameters.Add(TEXT("--follow")); // follow file renames
 		Parameters.Add(TEXT("--date=raw"));
 		Parameters.Add(TEXT("--name-status")); // relative filename at this revision, preceded by a status character
+		Parameters.Add(TEXT("--pretty=medium")); // make sure format matches expected in ParseLogResults
 		TArray<FString> Files;
 		Files.Add(*InFile);
 		if(bMergeConflict)
diff --git a/Engine/Plugins/Media/MediaPlayerEditor/Source/MediaPlayerEditor/Private/Shared/MediaPlayerEditorStyle.h b/Engine/Plugins/Media/MediaPlayerEditor/Source/MediaPlayerEditor/Private/Shared/MediaPlayerEditorStyle.h
--- a/Engine/Plugins/Media/MediaPlayerEditor/Source/MediaPlayerEditor/Private/Shared/MediaPlayerEditorStyle.h
+++ b/Engine/Plugins/Media/MediaPlayerEditor/Source/MediaPlayerEditor/Private/Shared/MediaPlayerEditorStyle.h
@@ -36,7 +36,7 @@ public:
 		Set("MediaPlayerEditor.SettingsButton", new IMAGE_BRUSH("btn_settings_16x", Icon12x12));
 
 		// misc
-		Set("MediaPlayerEditor.DragDropBorder", new BOX_BRUSH("border_dragdrop.png", 0.5f));
+		Set("MediaPlayerEditor.DragDropBorder", new BOX_BRUSH("border_dragdrop", 0.5f));
 		Set("MediaPlayerEditor.FileMediaSourcePrecached", new IMAGE_BRUSH("filemediasource_precached", Icon12x12));
 		Set("MediaPlayerEditor.MediaSourceOpened", new IMAGE_BRUSH("mediasource_opened", Icon8x8));
 
diff --git a/Engine/Plugins/Tests/EditorTests/Source/EditorTests/Private/UnrealEd/BlueprintAutomationTests.cpp b/Engine/Plugins/Tests/EditorTests/Source/EditorTests/Private/UnrealEd/BlueprintAutomationTests.cpp
--- a/Engine/Plugins/Tests/EditorTests/Source/EditorTests/Private/UnrealEd/BlueprintAutomationTests.cpp
+++ b/Engine/Plugins/Tests/EditorTests/Source/EditorTests/Private/UnrealEd/BlueprintAutomationTests.cpp
@@ -243,7 +243,7 @@ public:
 			FString Filename = Asset.ObjectPath.ToString();
 			//convert to full paths
 			Filename = FPackageName::LongPackageNameToFilename(Filename);
-			if (FAutomationTestFramework::GetInstance().ShouldTestContent(Filename))
+			if (FAutomationTestFramework::Get().ShouldTestContent(Filename))
 			{
 				// optionally discount already loaded assets
 				if (!bIgnoreLoaded || !Asset.IsAssetLoaded())
@@ -1244,7 +1244,7 @@ void FBlueprintRenameAndCloneTest::GetTests(TArray<FString>& OutBeautifiedNames,
 		FString AssetObjPath = Asset.ObjectPath.ToString();
 
 		FString const Filename = FPackageName::LongPackageNameToFilename(AssetObjPath);
-		if (!FAutomationTestFramework::GetInstance().ShouldTestContent(Filename))
+		if (!FAutomationTestFramework::Get().ShouldTestContent(Filename))
 		{
 			continue;
 		}
diff --git a/Engine/Plugins/Tests/EditorTests/Source/EditorTests/Private/UnrealEd/EditorAssetAutomationTests.cpp b/Engine/Plugins/Tests/EditorTests/Source/EditorTests/Private/UnrealEd/EditorAssetAutomationTests.cpp
--- a/Engine/Plugins/Tests/EditorTests/Source/EditorTests/Private/UnrealEd/EditorAssetAutomationTests.cpp
+++ b/Engine/Plugins/Tests/EditorTests/Source/EditorTests/Private/UnrealEd/EditorAssetAutomationTests.cpp
@@ -57,6 +57,8 @@
 #include "GameFramework/ForceFeedbackEffect.h"
 #include "Engine/ObjectLibrary.h"
 #include "GameFramework/TouchInterface.h"
+#include "Factories/DataAssetFactory.h"
+#include "Factories/CurveFactory.h"
 
 #define LOCTEXT_NAMESPACE "EditorAssetAutomationTests"
 
@@ -924,7 +926,7 @@ namespace ImportExportAssetHelper
 					Data.Width = OutImageSize.X;
 					Data.Height = OutImageSize.Y;
 					Data.Path = ScreenshotName;
-					FAutomationTestFramework::GetInstance().OnScreenshotCaptured().ExecuteIfBound(OutImageData, Data);
+					FAutomationTestFramework::Get().OnScreenshotCaptured().ExecuteIfBound(OutImageData, Data);
 				}
 
 				//Close the editor
diff --git a/Engine/Plugins/Tests/EditorTests/Source/EditorTests/Private/UnrealEd/EditorAutomationTests.cpp b/Engine/Plugins/Tests/EditorTests/Source/EditorTests/Private/UnrealEd/EditorAutomationTests.cpp
--- a/Engine/Plugins/Tests/EditorTests/Source/EditorTests/Private/UnrealEd/EditorAutomationTests.cpp
+++ b/Engine/Plugins/Tests/EditorTests/Source/EditorTests/Private/UnrealEd/EditorAutomationTests.cpp
@@ -101,7 +101,7 @@ bool PointLightDuplicationCommand::Update()
 void TakeLatentAutomationScreenshot(struct WindowScreenshotParameters ScreenshotParameters, FString BaseFileName, FString ScreenshotTitle, FString ScreenshotFolderName)
 {
 	//Update the screenshot name, then take a screenshot.
-	if (FAutomationTestFramework::GetInstance().IsScreenshotAllowed())
+	if (FAutomationTestFramework::Get().IsScreenshotAllowed())
 	{
 		//Update the screenshot name and get the location of where it will be saved.
 		ScreenshotParameters.ScreenshotName = TEXT("ScreenshotTitle");
@@ -250,7 +250,7 @@ void FLoadAllMapsInEditorTest::GetTests(TArray<FString>& OutBeautifiedNames, TAr
 		// Disregard filenames that don't have the map extension if we're in MAPSONLY mode.
 		if ( FPaths::GetExtension(Filename, true) == FPackageName::GetMapPackageExtension()) 
 		{
-			if (FAutomationTestFramework::GetInstance().ShouldTestContent(Filename))
+			if (FAutomationTestFramework::Get().ShouldTestContent(Filename))
 			{
 				if (!Filename.Contains(TEXT("/Engine/")))
 				{
@@ -282,7 +282,7 @@ bool FLoadAllMapsInEditorTest::RunTest(const FString& Parameters)
 	}
 	
 
-	const bool bTakeScreenshots = FAutomationTestFramework::GetInstance().IsScreenshotAllowed();
+	const bool bTakeScreenshots = FAutomationTestFramework::Get().IsScreenshotAllowed();
 	if( bTakeScreenshots )
 	{
 		//Find the main editor window
diff --git a/Engine/Plugins/Tests/EditorTests/Source/EditorTests/Private/UnrealEd/EditorOpenAssetAutomationTests.cpp b/Engine/Plugins/Tests/EditorTests/Source/EditorTests/Private/UnrealEd/EditorOpenAssetAutomationTests.cpp
--- a/Engine/Plugins/Tests/EditorTests/Source/EditorTests/Private/UnrealEd/EditorOpenAssetAutomationTests.cpp
+++ b/Engine/Plugins/Tests/EditorTests/Source/EditorTests/Private/UnrealEd/EditorOpenAssetAutomationTests.cpp
@@ -178,7 +178,7 @@ bool FOpenAllofAssetTypeCommand::Update()
 				if ( FAssetEditorManager::Get().FindEditorForAsset(Object, true) != NULL )
 				{
 					UE_LOG(LogEditorAutomationTests, Log, TEXT("Verified asset editor for: %s."), *AssetList.AssetPathList[AssetList.AssetListLength]);
-					UE_LOG(LogEditorAutomationTests, Display, TEXT("The editor successffully loaded for: %s."), *AssetList.AssetPathList[AssetList.AssetListLength]);
+					UE_LOG(LogEditorAutomationTests, Display, TEXT("The editor successfully loaded for: %s."), *AssetList.AssetPathList[AssetList.AssetListLength]);
 					if ( GShaderCompilingManager->IsCompiling() )
 					{
 						UE_LOG(LogEditorAutomationTests, Log, TEXT("Waiting for %i shaders to finish."), GShaderCompilingManager->GetNumRemainingJobs());
diff --git a/Engine/Plugins/Tests/EditorTests/Source/EditorTests/Private/UnrealEd/GeometryTests.cpp b/Engine/Plugins/Tests/EditorTests/Source/EditorTests/Private/UnrealEd/GeometryTests.cpp
--- a/Engine/Plugins/Tests/EditorTests/Source/EditorTests/Private/UnrealEd/GeometryTests.cpp
+++ b/Engine/Plugins/Tests/EditorTests/Source/EditorTests/Private/UnrealEd/GeometryTests.cpp
@@ -331,8 +331,8 @@ bool FBSPCubePlacement::RunTest(const FString& Parameters)
 
 	// Verify the BSP brushes that were created have been found.
 	TestEqual(TEXT("BSP Brushes Created"), BrushCount, 2);
+
 	TestNotNull<ABrush>(TEXT("Additive BSP is NULL.  Most likely it wasn't created."), AdditiveBSP);
-	TestNotNull<ABrush>(TEXT("Subtractive BSP is NULL.  Most likely it wasn't created."), SubtractiveBSP);
 
 	// Verify the location of the BSP.
 	if (AdditiveBSP)
@@ -353,7 +353,9 @@ bool FBSPCubePlacement::RunTest(const FString& Parameters)
 	{
 		TestEqual(TEXT("Subtractive Brush Vertex Count"), SubtractiveBSP->Brush->NumUniqueVertices, CUBE_VERTS_COUNT);
 	}
+	
 
+	
 	// Verify there are the correct amount of BSP surfaces are visible in the level.
 	TestEqual(TEXT("Surfaces Reported"), BSPModel->Surfs.Num(), 10);
 
diff --git a/Engine/Plugins/Tests/FunctionalTestingEditor/Source/FunctionalTestingEditor.Build.cs b/Engine/Plugins/Tests/FunctionalTestingEditor/Source/FunctionalTestingEditor.Build.cs
new file mode 100644
--- /dev/null
+++ b/Engine/Plugins/Tests/FunctionalTestingEditor/Source/FunctionalTestingEditor.Build.cs
@@ -0,0 +1,36 @@
+// Copyright 1998-2016 Epic Games, Inc. All Rights Reserved.
+
+using UnrealBuildTool;
+
+public class FunctionalTestingEditor : ModuleRules
+{
+    public FunctionalTestingEditor(TargetInfo Target)
+	{
+        PublicDependencyModuleNames.AddRange
+        (
+            new string[] {
+				"Core",
+                "InputCore",
+                "CoreUObject",
+                "SlateCore",
+                "Slate",
+                "EditorStyle",
+                "Engine",
+                "AssetRegistry"
+			}
+        );
+        
+        PrivateDependencyModuleNames.AddRange(
+             new string[] {
+                "Engine",
+                "UnrealEd",
+                "LevelEditor",
+                "SessionFrontend",
+                "FunctionalTesting",
+                "PlacementMode",
+                "WorkspaceMenuStructure",
+                "ScreenShotComparisonTools"
+            }
+         );
+	}
+}
diff --git a/Engine/Plugins/Tests/FunctionalTestingEditor/Source/Private/EditorAutomationModule.cpp b/Engine/Plugins/Tests/FunctionalTestingEditor/Source/Private/EditorAutomationModule.cpp
new file mode 100644
--- /dev/null
+++ b/Engine/Plugins/Tests/FunctionalTestingEditor/Source/Private/EditorAutomationModule.cpp
@@ -0,0 +1,180 @@
+// Copyright 1998-2016 Epic Games, Inc. All Rights Reserved.
+
+#include "EditorAutomationPrivatePCH.h"
+#include "EditorAutomationModule.h"
+#include "LevelEditor.h"
+#include "ISessionFrontendModule.h"
+#include "IPlacementModeModule.h"
+#include "FunctionalTest.h"
+#include "ScreenshotFunctionalTest.h"
+#include "AssetRegistryModule.h"
+#include "WorkspaceMenuStructureModule.h"
+#include "Framework/Docking/TabManager.h"
+
+#define LOCTEXT_NAMESPACE "EditorAutomation"
+
+void OpenMapAndFocusActor(const TArray<FString>& Args)
+{
+	if ( Args.Num() < 2 )
+	{
+		UE_LOG(LogConsoleResponse, Display, TEXT("Automate.OpenMapAndFocusActor failed\n"));
+		return;
+	}
+
+	FString AssetPath = Args[0];
+	FAssetRegistryModule& AssetRegistryModule = FModuleManager::LoadModuleChecked<FAssetRegistryModule>(TEXT("AssetRegistry"));
+
+	TArray<FAssetData> AllAssets;
+	AssetRegistryModule.Get().GetAssetsByPackageName(*AssetPath, AllAssets);
+
+	if ( AllAssets.Num() > 0 )
+	{
+		bool bIsWorlAlreadyOpened = false;
+		if ( UWorld* EditorWorld = GEditor->GetEditorWorldContext().World() )
+		{
+			if ( FAssetData(EditorWorld).PackageName == AllAssets[0].PackageName )
+			{
+				bIsWorlAlreadyOpened = true;
+			}
+		}
+
+		if ( !bIsWorlAlreadyOpened )
+		{
+			UObject* ObjectToEdit = AllAssets[0].GetAsset();
+			if ( ObjectToEdit )
+			{
+				GEditor->EditObject(ObjectToEdit);
+			}
+		}
+	}
+
+	if ( UWorld* EditorWorld = GEditor->GetEditorWorldContext().World() )
+	{
+		AActor* ActorToFocus = nullptr;
+		for ( int32 LevelIndex=0; LevelIndex < EditorWorld->GetNumLevels(); LevelIndex++ )
+		{
+			ULevel* Level = EditorWorld->GetLevel(LevelIndex);
+			ActorToFocus = static_cast<AActor*>( FindObjectWithOuter(Level, AActor::StaticClass(), *Args[1]) );
+			if ( ActorToFocus )
+			{
+				break;
+			}
+		}
+
+		if ( ActorToFocus )
+		{
+			GEditor->SelectNone(/*bNoteSelectionChange=*/ false, false, false);
+			GEditor->SelectActor(ActorToFocus, true, true);
+			GEditor->NoteSelectionChange();
+
+			const bool bActiveViewportOnly = false;
+			GEditor->MoveViewportCamerasToActor(*ActorToFocus, bActiveViewportOnly);
+		}
+	}
+}
+
+FAutoConsoleCommand OpenMapAndFocusActorCmd(
+	TEXT("Automate.OpenMapAndFocusActor"),
+	TEXT("Opens a map and focuses a particular actor by name."),
+	FConsoleCommandWithArgsDelegate::CreateStatic(OpenMapAndFocusActor)
+);
+
+class FEditorAutomationModule : public IEditorAutomationModule
+{
+	void StartupModule()
+	{
+		// make an extension to add the Orion function menu
+		Extender = MakeShareable(new FExtender());
+		Extender->AddMenuExtension(
+			"General",
+			EExtensionHook::After,
+			nullptr,
+			FMenuExtensionDelegate::CreateRaw(this, &FEditorAutomationModule::OnAutomationToolsMenuCreation));
+
+		// add the menu extension to the editor
+		FLevelEditorModule& LevelEditorModule = FModuleManager::LoadModuleChecked<FLevelEditorModule>("LevelEditor");
+		LevelEditorModule.GetMenuExtensibilityManager()->AddExtender(Extender);
+
+		FModuleManager::Get().OnModulesChanged().AddRaw(this, &FEditorAutomationModule::OnModulesChanged);
+
+		if ( IPlacementModeModule::IsAvailable() )
+		{
+			OnModulesChanged("PlacementMode", EModuleChangeReason::ModuleLoaded);
+		}
+	}
+
+	void ShutdownModule()
+	{
+		if ( FModuleManager::Get().IsModuleLoaded("LevelEditor") )
+		{
+			FLevelEditorModule& LevelEditorModule = FModuleManager::LoadModuleChecked<FLevelEditorModule>("LevelEditor");
+			LevelEditorModule.GetMenuExtensibilityManager()->RemoveExtender(Extender);
+		}
+
+		FModuleManager::Get().OnModulesChanged().RemoveAll(this);
+
+		if ( IPlacementModeModule::IsAvailable() )
+		{
+			IPlacementModeModule::Get().UnregisterPlacementCategory("Testing");
+		}
+	}
+
+	void OnAutomationToolsMenuCreation(FMenuBuilder& MenuBuilder)
+	{
+		MenuBuilder.BeginSection("Testing", LOCTEXT("Testing", "Testing"));
+		MenuBuilder.AddMenuEntry(
+			LOCTEXT("AutomationLabel", "Test Automation"),
+			LOCTEXT("Tooltip", "Launch the Testing Automation Frontend."),
+			FSlateIcon(FEditorStyle::GetStyleSetName(), "AutomationTools.MenuIcon"),
+			FUIAction(FExecuteAction::CreateStatic(&FEditorAutomationModule::OnShowAutomationFrontend)));
+		MenuBuilder.EndSection();
+
+		// TODO Come up with a way to hide this if no tools are registered.
+		if (WorkspaceMenu::GetMenuStructure().GetAutomationToolsCategory()->GetChildItems().Num() > 0)
+		{
+			MenuBuilder.AddSubMenu(
+				LOCTEXT("AutomationTools", "Automation Tools"),
+				LOCTEXT("AutomationToolsToolTip", "Assorted tools to help generate data for some of the automation tests."),
+				FNewMenuDelegate::CreateRaw(this, &FEditorAutomationModule::PopulateAutomationTools)
+			);
+		}
+	}
+
+	void PopulateAutomationTools(FMenuBuilder& MenuBuilder)
+	{
+		MenuBuilder.BeginSection("AutomationTools", LOCTEXT("AutomationTools", "Automation Tools"));
+		const bool bAutoAndOrphanedMenus = false;
+		FGlobalTabmanager::Get()->PopulateTabSpawnerMenu(MenuBuilder, WorkspaceMenu::GetMenuStructure().GetAutomationToolsCategory(), bAutoAndOrphanedMenus);
+		MenuBuilder.EndSection();
+	}
+
+	static void OnShowAutomationFrontend()
+	{
+		ISessionFrontendModule& SessionFrontend = FModuleManager::LoadModuleChecked<ISessionFrontendModule>("SessionFrontend");
+		SessionFrontend.InvokeSessionFrontend(FName("AutomationPanel"));
+	}
+
+	void OnModulesChanged(FName Module, EModuleChangeReason Reason)
+	{
+		if ( Module == TEXT("PlacementMode") && Reason == EModuleChangeReason::ModuleLoaded )
+		{
+			FPlacementCategoryInfo Info(
+				LOCTEXT("FunctionalTestingCategoryName", "Testing"),
+				"Testing",
+				TEXT("PMTesting"),
+				25
+				);
+
+			IPlacementModeModule::Get().RegisterPlacementCategory(Info);
+			IPlacementModeModule::Get().RegisterPlaceableItem(Info.UniqueHandle, MakeShareable(new FPlaceableItem(nullptr, FAssetData(AFunctionalTest::StaticClass()))));
+			IPlacementModeModule::Get().RegisterPlaceableItem(Info.UniqueHandle, MakeShareable(new FPlaceableItem(nullptr, FAssetData(AScreenshotFunctionalTest::StaticClass()))));
+		}
+	}
+
+private:
+	TSharedPtr<FExtender> Extender;
+};
+
+#undef LOCTEXT_NAMESPACE
+
+IMPLEMENT_MODULE(FEditorAutomationModule, EditorAutomation);
\ No newline at end of file
diff --git a/Engine/Plugins/Tests/FunctionalTestingEditor/Source/Private/EditorAutomationPrivatePCH.h b/Engine/Plugins/Tests/FunctionalTestingEditor/Source/Private/EditorAutomationPrivatePCH.h
new file mode 100644
--- /dev/null
+++ b/Engine/Plugins/Tests/FunctionalTestingEditor/Source/Private/EditorAutomationPrivatePCH.h
@@ -0,0 +1,6 @@
+// Copyright 1998-2016 Epic Games, Inc. All Rights Reserved.
+
+#pragma once
+
+#include "Core.h"
+#include "UnrealEd.h"
\ No newline at end of file
diff --git a/Engine/Plugins/Tests/FunctionalTestingEditor/Source/Public/EditorAutomationModule.h b/Engine/Plugins/Tests/FunctionalTestingEditor/Source/Public/EditorAutomationModule.h
new file mode 100644
--- /dev/null
+++ b/Engine/Plugins/Tests/FunctionalTestingEditor/Source/Public/EditorAutomationModule.h
@@ -0,0 +1,21 @@
+// Copyright 1998-2016 Epic Games, Inc. All Rights Reserved.
+
+#include "ModuleInterface.h"
+
+
+/**
+ * A module for adding automation exposure in the editor
+ */
+class IEditorAutomationModule : public IModuleInterface
+{
+public:
+	/**
+	* Called right after the module DLL has been loaded and the module object has been created
+	*/
+	virtual void StartupModule() = 0;
+
+	/**
+	* Called before the module is unloaded, right before the module object is destroyed.
+	*/
+	virtual void ShutdownModule() = 0;
+};
\ No newline at end of file
diff --git a/Engine/Source/Developer/AssetTools/Private/AssetTools.cpp b/Engine/Source/Developer/AssetTools/Private/AssetTools.cpp
--- a/Engine/Source/Developer/AssetTools/Private/AssetTools.cpp
+++ b/Engine/Source/Developer/AssetTools/Private/AssetTools.cpp
@@ -16,6 +16,7 @@
 #include "UnrealExporter.h"
 #include "SNotificationList.h"
 #include "NotificationManager.h"
+#include "AutomatedAssetImportData.h"
 
 #define LOCTEXT_NAMESPACE "AssetTools"
 
@@ -564,6 +565,19 @@
 	return ReturnObjects;
 }
 
+TArray<UObject*> FAssetTools::ImportAssetsAutomated(const UAutomatedAssetImportData& ImportData) const
+{
+	FAssetImportParams Params;
+
+	Params.bAutomated = true;
+	Params.bForceOverrideExisting = ImportData.bReplaceExisting;
+	Params.bSyncToBrowser = false;
+	Params.SpecifiedFactory = ImportData.Factory;
+	Params.ImportData = &ImportData;
+
+	return ImportAssetsInternal(ImportData.Filenames, ImportData.DestinationPath, Params);
+}
+
 void FAssetTools::ExpandDirectories(const TArray<FString>& Files, const FString& DestinationPath, TArray<TPair<FString, FString>>& FilesAndDestinations) const
 {
 	// Iterate through all files in the list, if any folders are found, recurse and expand them.
@@ -600,415 +614,18 @@
 	}
 }
 
-TArray<UObject*> FAssetTools::ImportAssets(const TArray<FString>& Files, const FString& RootDestinationPath, UFactory* SpecifiedFactory, bool bSyncToBrowser ) const
+TArray<UObject*> FAssetTools::ImportAssets(const TArray<FString>& Files, const FString& RootDestinationPath, UFactory* SpecifiedFactory, bool bSyncToBrowser) const
 {
-	TArray<UObject*> ReturnObjects;
-	TMap< FString, TArray<UFactory*> > ExtensionToFactoriesMap;
+	const bool bForceOverrideExisting = false;
 
-	FScopedSlowTask SlowTask(Files.Num() + 3, LOCTEXT("ImportSlowTask", "Importing"));
-	SlowTask.MakeDialog();
+	FAssetImportParams Params;
 
-	// Reset the 'Do you want to overwrite the existing object?' Yes to All / No to All prompt, to make sure the
-	// user gets a chance to select something
-	UFactory::ResetState();
+	Params.bAutomated = false;
+	Params.bForceOverrideExisting = false;
+	Params.bSyncToBrowser = bSyncToBrowser;
+	Params.SpecifiedFactory = SpecifiedFactory;
 
-	SlowTask.EnterProgressFrame();
-
-	TArray<TPair<FString, FString>> FilesAndDestinations;
-	ExpandDirectories(Files, RootDestinationPath, FilesAndDestinations);
-
-	SlowTask.EnterProgressFrame(1, LOCTEXT("Import_DeterminingImportTypes", "Determining asset types"));
-
-	if (SpecifiedFactory == nullptr)
-	{	
-		// First instantiate one factory for each file extension encountered that supports the extension
-		// @todo import: gmp: show dialog in case of multiple matching factories
-		for( TObjectIterator<UClass> ClassIt; ClassIt; ++ClassIt )
-		{
-			if (!(*ClassIt)->IsChildOf(UFactory::StaticClass()) || ((*ClassIt)->HasAnyClassFlags(CLASS_Abstract)))
-			{
-				continue;
-			}
-
-			UFactory* Factory = Cast<UFactory>( (*ClassIt)->GetDefaultObject() );
-
-			if (!Factory->bEditorImport)
-			{
-				continue;
-			}
-
-			TArray<FString> FactoryExtensions;
-			Factory->GetSupportedFileExtensions(FactoryExtensions);
-
-			for (auto& FileDest : FilesAndDestinations)
-			{
-				const FString FileExtension = FPaths::GetExtension(FileDest.Key);
-	
-				// Case insensitive string compare with supported formats of this factory
-				if ( FactoryExtensions.Contains(FileExtension) )
-				{
-					TArray<UFactory*>& ExistingFactories = ExtensionToFactoriesMap.FindOrAdd(FileExtension);
-							
-					// Do not remap extensions, just reuse the existing UFactory.
-					// There may be multiple UFactories, so we will keep track of all of them
-					bool bFactoryAlreadyInMap = false;
-					for ( auto FoundFactoryIt = ExistingFactories.CreateConstIterator(); FoundFactoryIt; ++FoundFactoryIt )
-					{
-						if ( (*FoundFactoryIt)->GetClass() == Factory->GetClass() )
-						{
-							bFactoryAlreadyInMap = true;
-							break;
-						}
-					}
-	
-					if ( !bFactoryAlreadyInMap )
-					{
-						// We found a factory for this file, it can be imported!
-						// Create a new factory of the same class and make sure it doesn't get GCed.
-						// The object will be removed from the root set at the end of this function.
-						UFactory* NewFactory = NewObject<UFactory>(GetTransientPackage(), Factory->GetClass());
-						if ( NewFactory->ConfigureProperties() )
-						{
-							NewFactory->AddToRoot();
-							ExistingFactories.Add(NewFactory);
-						}
-					}
-				}
-			}
-		}
-	}
-	else if( SpecifiedFactory->bEditorImport )
-	{
-		TArray<FString> FactoryExtensions;
-		SpecifiedFactory->GetSupportedFileExtensions(FactoryExtensions);
-
-		for (auto FileIt = Files.CreateConstIterator(); FileIt; ++FileIt)
-		{
-			const FString FileExtension = FPaths::GetExtension(*FileIt);
-
-			// Case insensitive string compare with supported formats of this factory
-			if (!FactoryExtensions.Contains(FileExtension))
-			{
-				continue;
-			}
-
-			TArray<UFactory*>& ExistingFactories = ExtensionToFactoriesMap.FindOrAdd(FileExtension);
-
-			// Do not remap extensions, just reuse the existing UFactory.
-			// There may be multiple UFactories, so we will keep track of all of them
-			bool bFactoryAlreadyInMap = false;
-			for ( auto FoundFactoryIt = ExistingFactories.CreateConstIterator(); FoundFactoryIt; ++FoundFactoryIt )
-			{
-				if ( (*FoundFactoryIt)->GetClass() == SpecifiedFactory->GetClass() )
-				{
-					bFactoryAlreadyInMap = true;
-					break;
-				}
-			}
-
-			if ( !bFactoryAlreadyInMap )
-			{
-				// We found a factory for this file, it can be imported!
-				// Create a new factory of the same class and make sure it doesnt get GCed.
-				// The object will be removed from the root set at the end of this function.
-				UFactory* NewFactory = NewObject<UFactory>(GetTransientPackage(), SpecifiedFactory->GetClass());
-				if ( NewFactory->ConfigureProperties() )
-				{
-					NewFactory->AddToRoot();
-					ExistingFactories.Add(NewFactory);
-				}
-			}
-		}
-	}
-
-	// We need to sort the factories so that they get tested in priority order
-	for (auto& ExtensionToFactories : ExtensionToFactoriesMap)
-	{
-		ExtensionToFactories.Value.Sort(&UFactory::SortFactoriesByPriority);
-	}
-
-	// Some flags to keep track of what the user decided when asked about overwriting or replacing
-	bool bOverwriteAll = false;
-	bool bReplaceAll = false;
-	bool bDontOverwriteAny = false;
-	bool bDontReplaceAny = false;
-
-	// Now iterate over the input files and use the same factory object for each file with the same extension
-	for ( int32 FileIdx = 0; FileIdx < FilesAndDestinations.Num(); ++FileIdx )
-	{
-		const FString& Filename = FilesAndDestinations[FileIdx].Key;
-		const FString& DestinationPath = FilesAndDestinations[FileIdx].Value;
-
-		SlowTask.EnterProgressFrame(1, FText::Format(LOCTEXT("Import_ImportingFile", "Importing \"{0}\"..."), FText::FromString(FPaths::GetBaseFilename(Filename))));
-
-		FString FileExtension = FPaths::GetExtension(Filename);
-
-		const TArray<UFactory*>* FactoriesPtr = ExtensionToFactoriesMap.Find(FileExtension);
-		UFactory* Factory = nullptr;
-		//When doing automationtest we can setup factory option and we need to make sure we use the specified one
-		if (GIsAutomationTesting && SpecifiedFactory && SpecifiedFactory->FactoryCanImport(Filename))
-		{
-			Factory = SpecifiedFactory;
-		}
-		else if ( FactoriesPtr )
-		{
-			const TArray<UFactory*>& Factories = *FactoriesPtr;
-
-			// Handle the potential of multiple factories being found
-			if( Factories.Num() > 0 )
-			{
-				Factory = Factories[0];
-
-				for( auto FactoryIt = Factories.CreateConstIterator(); FactoryIt; ++FactoryIt )
-				{
-					UFactory* TestFactory = *FactoryIt;
-					if( TestFactory->FactoryCanImport( Filename ) )
-					{
-						Factory = TestFactory;
-						break;
-					}
-				}
-			}
-		}
-		else
-		{
-			if(FEngineAnalytics::IsAvailable())
-			{
-				TArray<FAnalyticsEventAttribute> Attribs;
-				Attribs.Add(FAnalyticsEventAttribute(TEXT("FileExtension"), FileExtension));
-
-				FEngineAnalytics::GetProvider().RecordEvent(TEXT("Editor.Usage.ImportFailed"), Attribs);
-			}
-
-			const FText Message = FText::Format( LOCTEXT("ImportFailed_UnknownExtension", "Failed to import '{0}'. Unknown extension '{1}'."), FText::FromString( Filename ), FText::FromString( FileExtension ) );
-			FNotificationInfo Info(Message);
-			Info.ExpireDuration = 3.0f;
-			Info.bUseLargeFont = false;
-			Info.bFireAndForget = true;
-			Info.bUseSuccessFailIcons = true;
-			FSlateNotificationManager::Get().AddNotification(Info)->SetCompletionState(SNotificationItem::CS_Fail);
-
-			UE_LOG(LogAssetTools, Warning, TEXT("%s"), *Message.ToString() );
-		}
-
-		if ( Factory != nullptr )
-		{
-			UClass* ImportAssetType = Factory->SupportedClass;
-			bool bImportSucceeded = false;
-			bool bImportWasCancelled = false;
-			FDateTime ImportStartTime = FDateTime::UtcNow();
-
-			FString Name = ObjectTools::SanitizeObjectName(FPaths::GetBaseFilename(Filename));
-			FString PackageName = DestinationPath + TEXT("/") + Name;
-
-			// We can not create assets that share the name of a map file in the same location
-			if ( FEditorFileUtils::IsMapPackageAsset(PackageName) )
-			{
-				const FText Message = FText::Format( LOCTEXT("AssetNameInUseByMap", "You can not create an asset named '{0}' because there is already a map file with this name in this folder."), FText::FromString( Name ) );
-				FMessageDialog::Open( EAppMsgType::Ok, Message );
-				UE_LOG(LogAssetTools, Warning, TEXT("%s"), *Message.ToString());
-				OnNewImportRecord(ImportAssetType, FileExtension, bImportSucceeded, bImportWasCancelled, ImportStartTime);
-				continue;
-			}
-
-			UPackage* Pkg = CreatePackage(nullptr, *PackageName);
-			if ( !ensure(Pkg) )
-			{
-				// Failed to create the package to hold this asset for some reason
-				OnNewImportRecord(ImportAssetType, FileExtension, bImportSucceeded, bImportWasCancelled, ImportStartTime);
-				continue;
-			}
-
-			// Make sure the destination package is loaded
-			Pkg->FullyLoad();
-
-			// Check for an existing object
-			UObject* ExistingObject = StaticFindObject( UObject::StaticClass(), Pkg, *Name );
-			if( ExistingObject != nullptr )
-			{
-				// If the existing object is one of the imports we've just created we can't replace or overwrite it
-				if (ReturnObjects.Contains(ExistingObject))
-				{
-					if( ImportAssetType == nullptr )
-					{
-						// The factory probably supports multiple types and cant be determined yet without asking the user or actually loading it
-						// We just need to generate an unused name so object should do fine.
-						ImportAssetType = UObject::StaticClass();
-					}
-					// generate a unique name for this import
-					Name = MakeUniqueObjectName(Pkg, ImportAssetType, *Name).ToString();
-				}
-				else
-				{
-					// If the object is supported by the factory we are using, ask if we want to overwrite the asset
-					// Otherwise, prompt to replace the object
-					if (Factory->DoesSupportClass(ExistingObject->GetClass()))
-					{
-						// The factory can overwrite this object, ask if that is okay, unless "Yes To All" or "No To All" was already selected
-						EAppReturnType::Type UserResponse;
-
-						if (bOverwriteAll || GIsAutomationTesting)
-						{
-							UserResponse = EAppReturnType::YesAll;
-						}
-						else if (bDontOverwriteAny)
-						{
-							UserResponse = EAppReturnType::NoAll;
-						}
-						else
-						{
-							UserResponse = FMessageDialog::Open(
-								EAppMsgType::YesNoYesAllNoAll,
-								FText::Format(LOCTEXT("ImportObjectAlreadyExists_SameClass", "Do you want to overwrite the existing asset?\n\nAn asset already exists at the import location: {0}"), FText::FromString(PackageName)));
-
-							bOverwriteAll = UserResponse == EAppReturnType::YesAll;
-							bDontOverwriteAny = UserResponse == EAppReturnType::NoAll;
-						}
-
-						const bool bWantOverwrite = UserResponse == EAppReturnType::Yes || UserResponse == EAppReturnType::YesAll;
-
-						if (!bWantOverwrite)
-						{
-							// User chose not to replace the package
-							bImportWasCancelled = true;
-							OnNewImportRecord(ImportAssetType, FileExtension, bImportSucceeded, bImportWasCancelled, ImportStartTime);
-							continue;
-						}
-					}
-					else
-					{
-						// The factory can't overwrite this asset, ask if we should delete the object then import the new one. Only do this if "Yes To All" or "No To All" was not already selected.
-						EAppReturnType::Type UserResponse;
-
-						if (bReplaceAll)
-						{
-							UserResponse = EAppReturnType::YesAll;
-						}
-						else if (bDontReplaceAny)
-						{
-							UserResponse = EAppReturnType::NoAll;
-						}
-						else
-						{
-							UserResponse = FMessageDialog::Open(
-								EAppMsgType::YesNoYesAllNoAll,
-								FText::Format(LOCTEXT("ImportObjectAlreadyExists_DifferentClass", "Do you want to replace the existing asset?\n\nAn asset already exists at the import location: {0}"), FText::FromString(PackageName)));
-
-							bReplaceAll = UserResponse == EAppReturnType::YesAll;
-							bDontReplaceAny = UserResponse == EAppReturnType::NoAll;
-						}
-
-						const bool bWantReplace = UserResponse == EAppReturnType::Yes || UserResponse == EAppReturnType::YesAll;
-
-						if (bWantReplace)
-						{
-							// Delete the existing object
-							int32 NumObjectsDeleted = 0;
-							TArray< UObject* > ObjectsToDelete;
-							ObjectsToDelete.Add(ExistingObject);
-
-							// If the user forcefully deletes the package, all sorts of things could become invalidated,
-							// the Pkg pointer might be killed even though it was added to the root.
-							TWeakObjectPtr<UPackage> WeakPkg(Pkg);
-							
-							// Dont let the package get garbage collected (just in case we are deleting the last asset in the package)
-							Pkg->AddToRoot();
-							NumObjectsDeleted = ObjectTools::DeleteObjects(ObjectsToDelete, /*bShowConfirmation=*/false);
-							
-							// If the weak package ptr is still valid, it should then be safe to remove it from the root.
-							if ( WeakPkg.IsValid() )
-							{
-								Pkg->RemoveFromRoot();
-							}
-
-							const FString QualifiedName = PackageName + TEXT(".") + Name;
-							FText Reason;
-							if (NumObjectsDeleted == 0 || !IsUniqueObjectName(*QualifiedName, ANY_PACKAGE, Reason))
-							{
-								// Original object couldn't be deleted
-								const FText Message = FText::Format(LOCTEXT("ImportDeleteFailed", "Failed to delete '{0}'. The asset is referenced by other content."), FText::FromString(PackageName));
-								FMessageDialog::Open(EAppMsgType::Ok, Message);
-								UE_LOG(LogAssetTools, Warning, TEXT("%s"), *Message.ToString());
-								OnNewImportRecord(ImportAssetType, FileExtension, bImportSucceeded, bImportWasCancelled, ImportStartTime);
-								continue;
-							}
-							else
-							{
-								// succeed, recreate package since it has been deleted
-								Pkg = CreatePackage(nullptr, *PackageName);
-							}
-						}
-						else
-						{
-							// User chose not to replace the package
-							bImportWasCancelled = true;
-							OnNewImportRecord(ImportAssetType, FileExtension, bImportSucceeded, bImportWasCancelled, ImportStartTime);
-							continue;
-						}
-					}
-				}
-			}
-
-			// Check for a package that was marked for delete in source control
-			if ( !CheckForDeletedPackage(Pkg) )
-			{
-				OnNewImportRecord(ImportAssetType, FileExtension, bImportSucceeded, bImportWasCancelled, ImportStartTime);
-				continue;
-			}
-
-			ImportAssetType = Factory->ResolveSupportedClass();
-			UObject* Result = Factory->ImportObject(ImportAssetType, Pkg, FName(*Name), RF_Public | RF_Standalone, Filename, nullptr, bImportWasCancelled);
-
-			// Do not report any error if the operation was canceled.
-			if(!bImportWasCancelled)
-			{
-				if ( Result )
-				{
-					ReturnObjects.Add( Result );
-
-					// Notify the asset registry
-					FAssetRegistryModule::AssetCreated(Result);
-					GEditor->BroadcastObjectReimported(Result);
-
-					bImportSucceeded = true;
-				}
-				else
-				{
-					const FText Message = FText::Format( LOCTEXT("ImportFailed_Generic", "Failed to import '{0}'. Failed to create asset '{1}'.\nPlease see Output Log for details."), FText::FromString( Filename ), FText::FromString( PackageName ) );
-					FMessageDialog::Open( EAppMsgType::Ok, Message );
-					UE_LOG(LogAssetTools, Warning, TEXT("%s"), *Message.ToString());
-				}
-			}
-
-			// Refresh the supported class.  Some factories (e.g. FBX) only resolve their type after reading the file
-			ImportAssetType = Factory->ResolveSupportedClass();
-			OnNewImportRecord(ImportAssetType, FileExtension, bImportSucceeded, bImportWasCancelled, ImportStartTime);
-		}
-		else
-		{
-			// A factory or extension was not found. The extension warning is above. If a factory was not found, the user likely canceled a factory configuration dialog.
-		}
-	}
-
-	SlowTask.EnterProgressFrame(1);
-
-	// Clean up and remove the factories we created from the root set
-	for ( auto ExtensionIt = ExtensionToFactoriesMap.CreateConstIterator(); ExtensionIt; ++ExtensionIt)
-	{
-		for ( auto FactoryIt = ExtensionIt.Value().CreateConstIterator(); FactoryIt; ++FactoryIt )
-		{
-			(*FactoryIt)->CleanUp();
-			(*FactoryIt)->RemoveFromRoot();
-		}
-	}
-
-	// Sync content browser to the newly created assets
-	if ( ReturnObjects.Num() && (bSyncToBrowser != false ) )
-	{
-		FAssetTools::Get().SyncBrowserToAssets(ReturnObjects);
-	}
-
-	return ReturnObjects;
+	return ImportAssetsInternal(Files, RootDestinationPath, Params);
 }
 
 void FAssetTools::CreateUniqueAssetName(const FString& InBasePackageName, const FString& InSuffix, FString& OutPackageName, FString& OutAssetName) const
@@ -1385,4 +1002,433 @@
 	}
 }
 
+TArray<UObject*> FAssetTools::ImportAssetsInternal(const TArray<FString>& Files, const FString& RootDestinationPath, const FAssetImportParams& Params) const
+{
+	UFactory* SpecifiedFactory = Params.SpecifiedFactory;
+	const bool bForceOverrideExisting = Params.bForceOverrideExisting;
+	const bool bSyncToBrowser = Params.bSyncToBrowser;
+	const bool bAutomatedImport = Params.bAutomated || GIsAutomationTesting;
+
+	TArray<UObject*> ReturnObjects;
+	TMap< FString, TArray<UFactory*> > ExtensionToFactoriesMap;
+
+	FScopedSlowTask SlowTask(Files.Num() + 3, LOCTEXT("ImportSlowTask", "Importing"));
+	SlowTask.MakeDialog();
+
+
+	SlowTask.EnterProgressFrame();
+
+	TArray<TPair<FString, FString>> FilesAndDestinations;
+	ExpandDirectories(Files, RootDestinationPath, FilesAndDestinations);
+
+	SlowTask.EnterProgressFrame(1, LOCTEXT("Import_DeterminingImportTypes", "Determining asset types"));
+
+
+	if(SpecifiedFactory == nullptr)
+	{
+		// First instantiate one factory for each file extension encountered that supports the extension
+		// @todo import: gmp: show dialog in case of multiple matching factories
+		for(TObjectIterator<UClass> ClassIt; ClassIt; ++ClassIt)
+		{
+			if(!(*ClassIt)->IsChildOf(UFactory::StaticClass()) || ((*ClassIt)->HasAnyClassFlags(CLASS_Abstract)))
+			{
+				continue;
+			}
+
+			UFactory* Factory = Cast<UFactory>((*ClassIt)->GetDefaultObject());
+
+			if(!Factory->bEditorImport)
+			{
+				continue;
+			}
+
+			TArray<FString> FactoryExtensions;
+			Factory->GetSupportedFileExtensions(FactoryExtensions);
+
+			for(auto& FileDest : FilesAndDestinations)
+			{
+				const FString FileExtension = FPaths::GetExtension(FileDest.Key);
+
+				// Case insensitive string compare with supported formats of this factory
+				if(FactoryExtensions.Contains(FileExtension))
+				{
+					TArray<UFactory*>& ExistingFactories = ExtensionToFactoriesMap.FindOrAdd(FileExtension);
+
+					// Do not remap extensions, just reuse the existing UFactory.
+					// There may be multiple UFactories, so we will keep track of all of them
+					bool bFactoryAlreadyInMap = false;
+					for(auto FoundFactoryIt = ExistingFactories.CreateConstIterator(); FoundFactoryIt; ++FoundFactoryIt)
+					{
+						if((*FoundFactoryIt)->GetClass() == Factory->GetClass())
+						{
+							bFactoryAlreadyInMap = true;
+							break;
+						}
+					}
+
+					if(!bFactoryAlreadyInMap)
+					{
+						// We found a factory for this file, it can be imported!
+						// Create a new factory of the same class and make sure it doesn't get GCed.
+						// The object will be removed from the root set at the end of this function.
+						UFactory* NewFactory = NewObject<UFactory>(GetTransientPackage(), Factory->GetClass());
+						if(NewFactory->ConfigureProperties())
+						{
+							NewFactory->AddToRoot();
+							ExistingFactories.Add(NewFactory);
+						}
+					}
+				}
+			}
+		}
+	}
+	else if(SpecifiedFactory->bEditorImport && !bAutomatedImport) 
+	{
+
+		TArray<FString> FactoryExtensions;
+		SpecifiedFactory->GetSupportedFileExtensions(FactoryExtensions);
+
+		for(auto FileIt = Files.CreateConstIterator(); FileIt; ++FileIt)
+		{
+			const FString FileExtension = FPaths::GetExtension(*FileIt);
+
+			// Case insensitive string compare with supported formats of this factory
+			if(!FactoryExtensions.Contains(FileExtension))
+			{
+				continue;
+			}
+
+			TArray<UFactory*>& ExistingFactories = ExtensionToFactoriesMap.FindOrAdd(FileExtension);
+
+			// Do not remap extensions, just reuse the existing UFactory.
+			// There may be multiple UFactories, so we will keep track of all of them
+			bool bFactoryAlreadyInMap = false;
+			for(auto FoundFactoryIt = ExistingFactories.CreateConstIterator(); FoundFactoryIt; ++FoundFactoryIt)
+			{
+				if((*FoundFactoryIt)->GetClass() == SpecifiedFactory->GetClass())
+				{
+					bFactoryAlreadyInMap = true;
+					break;
+				}
+			}
+
+			if(!bFactoryAlreadyInMap)
+			{
+				// We found a factory for this file, it can be imported!
+				// Create a new factory of the same class and make sure it doesnt get GCed.
+				// The object will be removed from the root set at the end of this function.
+				UFactory* NewFactory = NewObject<UFactory>(GetTransientPackage(), SpecifiedFactory->GetClass());
+				if(NewFactory->ConfigureProperties())
+				{
+					NewFactory->AddToRoot();
+					ExistingFactories.Add(NewFactory);
+				}
+			}
+		}
+	}
+
+	// We need to sort the factories so that they get tested in priority order
+	for(auto& ExtensionToFactories : ExtensionToFactoriesMap)
+	{
+		ExtensionToFactories.Value.Sort(&UFactory::SortFactoriesByPriority);
+	}
+
+	// Some flags to keep track of what the user decided when asked about overwriting or replacing
+	bool bOverwriteAll = false;
+	bool bReplaceAll = false;
+	bool bDontOverwriteAny = false;
+	bool bDontReplaceAny = false;
+
+	// Now iterate over the input files and use the same factory object for each file with the same extension
+	for(int32 FileIdx = 0; FileIdx < FilesAndDestinations.Num(); ++FileIdx)
+	{
+		const FString& Filename = FilesAndDestinations[FileIdx].Key;
+		const FString& DestinationPath = FilesAndDestinations[FileIdx].Value;
+
+		SlowTask.EnterProgressFrame(1, FText::Format(LOCTEXT("Import_ImportingFile", "Importing \"{0}\"..."), FText::FromString(FPaths::GetBaseFilename(Filename))));
+
+		FString FileExtension = FPaths::GetExtension(Filename);
+
+		const TArray<UFactory*>* FactoriesPtr = ExtensionToFactoriesMap.Find(FileExtension);
+		UFactory* Factory = nullptr;
+		// Assume that for automated import, the user knows exactly what factory to use if it exists
+		if(bAutomatedImport && SpecifiedFactory && SpecifiedFactory->FactoryCanImport(Filename))
+		{
+			Factory = SpecifiedFactory;
+		}
+		else if(FactoriesPtr)
+		{
+			const TArray<UFactory*>& Factories = *FactoriesPtr;
+
+			// Handle the potential of multiple factories being found
+			if(Factories.Num() > 0)
+			{
+				Factory = Factories[0];
+
+				for(auto FactoryIt = Factories.CreateConstIterator(); FactoryIt; ++FactoryIt)
+				{
+					UFactory* TestFactory = *FactoryIt;
+					if(TestFactory->FactoryCanImport(Filename))
+					{
+						Factory = TestFactory;
+						break;
+					}
+				}
+			}
+		}
+		else
+		{
+			if(FEngineAnalytics::IsAvailable())
+			{
+				TArray<FAnalyticsEventAttribute> Attribs;
+				Attribs.Add(FAnalyticsEventAttribute(TEXT("FileExtension"), FileExtension));
+
+				FEngineAnalytics::GetProvider().RecordEvent(TEXT("Editor.Usage.ImportFailed"), Attribs);
+			}
+
+			const FText Message = FText::Format(LOCTEXT("ImportFailed_UnknownExtension", "Failed to import '{0}'. Unknown extension '{1}'."), FText::FromString(Filename), FText::FromString(FileExtension));
+			FNotificationInfo Info(Message);
+			Info.ExpireDuration = 3.0f;
+			Info.bUseLargeFont = false;
+			Info.bFireAndForget = true;
+			Info.bUseSuccessFailIcons = true;
+			FSlateNotificationManager::Get().AddNotification(Info)->SetCompletionState(SNotificationItem::CS_Fail);
+
+			UE_LOG(LogAssetTools, Warning, TEXT("%s"), *Message.ToString());
+		}
+
+		if(Factory != nullptr)
+		{
+			// Reset the 'Do you want to overwrite the existing object?' Yes to All / No to All prompt, to make sure the
+			// user gets a chance to select something
+			Factory->ResetState();
+
+			UClass* ImportAssetType = Factory->SupportedClass;
+			bool bImportSucceeded = false;
+			bool bImportWasCancelled = false;
+			FDateTime ImportStartTime = FDateTime::UtcNow();
+
+			FString Name = ObjectTools::SanitizeObjectName(FPaths::GetBaseFilename(Filename));
+			FString PackageName = DestinationPath + TEXT("/") + Name;
+
+			// We can not create assets that share the name of a map file in the same location
+			if(FEditorFileUtils::IsMapPackageAsset(PackageName))
+			{
+				const FText Message = FText::Format(LOCTEXT("AssetNameInUseByMap", "You can not create an asset named '{0}' because there is already a map file with this name in this folder."), FText::FromString(Name));
+				if(!bAutomatedImport)
+				{
+					FMessageDialog::Open(EAppMsgType::Ok, Message);
+				}
+				UE_LOG(LogAssetTools, Warning, TEXT("%s"), *Message.ToString());
+				OnNewImportRecord(ImportAssetType, FileExtension, bImportSucceeded, bImportWasCancelled, ImportStartTime);
+				continue;
+			}
+
+			UPackage* Pkg = CreatePackage(nullptr, *PackageName);
+			if(!ensure(Pkg))
+			{
+				// Failed to create the package to hold this asset for some reason
+				OnNewImportRecord(ImportAssetType, FileExtension, bImportSucceeded, bImportWasCancelled, ImportStartTime);
+				continue;
+			}
+
+			// Make sure the destination package is loaded
+			Pkg->FullyLoad();
+
+			// Check for an existing object
+			UObject* ExistingObject = StaticFindObject(UObject::StaticClass(), Pkg, *Name);
+			if(ExistingObject != nullptr)
+			{
+				// If the existing object is one of the imports we've just created we can't replace or overwrite it
+				if(ReturnObjects.Contains(ExistingObject))
+				{
+					if(ImportAssetType == nullptr)
+					{
+						// The factory probably supports multiple types and cant be determined yet without asking the user or actually loading it
+						// We just need to generate an unused name so object should do fine.
+						ImportAssetType = UObject::StaticClass();
+					}
+					// generate a unique name for this import
+					Name = MakeUniqueObjectName(Pkg, ImportAssetType, *Name).ToString();
+				}
+				else
+				{
+					// If the object is supported by the factory we are using, ask if we want to overwrite the asset
+					// Otherwise, prompt to replace the object
+					if(Factory->DoesSupportClass(ExistingObject->GetClass()))
+					{
+						// The factory can overwrite this object, ask if that is okay, unless "Yes To All" or "No To All" was already selected
+						EAppReturnType::Type UserResponse;
+
+						if(bForceOverrideExisting || bOverwriteAll || GIsAutomationTesting)
+						{
+							UserResponse = EAppReturnType::YesAll;
+						}
+						else if(bDontOverwriteAny)
+						{
+							UserResponse = EAppReturnType::NoAll;
+						}
+						else
+						{
+							UserResponse = FMessageDialog::Open(
+								EAppMsgType::YesNoYesAllNoAll,
+								FText::Format(LOCTEXT("ImportObjectAlreadyExists_SameClass", "Do you want to overwrite the existing asset?\n\nAn asset already exists at the import location: {0}"), FText::FromString(PackageName)));
+
+							bOverwriteAll = UserResponse == EAppReturnType::YesAll;
+							bDontOverwriteAny = UserResponse == EAppReturnType::NoAll;
+						}
+
+						const bool bWantOverwrite = UserResponse == EAppReturnType::Yes || UserResponse == EAppReturnType::YesAll;
+
+						if(!bWantOverwrite)
+						{
+							// User chose not to replace the package
+							bImportWasCancelled = true;
+							OnNewImportRecord(ImportAssetType, FileExtension, bImportSucceeded, bImportWasCancelled, ImportStartTime);
+							continue;
+						}
+					}
+					else if(!bAutomatedImport)
+					{
+						// The factory can't overwrite this asset, ask if we should delete the object then import the new one. Only do this if "Yes To All" or "No To All" was not already selected.
+						EAppReturnType::Type UserResponse;
+
+						if(bReplaceAll)
+						{
+							UserResponse = EAppReturnType::YesAll;
+						}
+						else if(bDontReplaceAny)
+						{
+							UserResponse = EAppReturnType::NoAll;
+						}
+						else
+						{
+							UserResponse = FMessageDialog::Open(
+								EAppMsgType::YesNoYesAllNoAll,
+								FText::Format(LOCTEXT("ImportObjectAlreadyExists_DifferentClass", "Do you want to replace the existing asset?\n\nAn asset already exists at the import location: {0}"), FText::FromString(PackageName)));
+
+							bReplaceAll = UserResponse == EAppReturnType::YesAll;
+							bDontReplaceAny = UserResponse == EAppReturnType::NoAll;
+						}
+
+						const bool bWantReplace = UserResponse == EAppReturnType::Yes || UserResponse == EAppReturnType::YesAll;
+
+						if(bWantReplace)
+						{
+							// Delete the existing object
+							int32 NumObjectsDeleted = 0;
+							TArray< UObject* > ObjectsToDelete;
+							ObjectsToDelete.Add(ExistingObject);
+
+							// If the user forcefully deletes the package, all sorts of things could become invalidated,
+							// the Pkg pointer might be killed even though it was added to the root.
+							TWeakObjectPtr<UPackage> WeakPkg(Pkg);
+
+							// Dont let the package get garbage collected (just in case we are deleting the last asset in the package)
+							Pkg->AddToRoot();
+							NumObjectsDeleted = ObjectTools::DeleteObjects(ObjectsToDelete, /*bShowConfirmation=*/false);
+
+							// If the weak package ptr is still valid, it should then be safe to remove it from the root.
+							if(WeakPkg.IsValid())
+							{
+								Pkg->RemoveFromRoot();
+							}
+
+							const FString QualifiedName = PackageName + TEXT(".") + Name;
+							FText Reason;
+							if(NumObjectsDeleted == 0 || !IsUniqueObjectName(*QualifiedName, ANY_PACKAGE, Reason))
+							{
+								// Original object couldn't be deleted
+								const FText Message = FText::Format(LOCTEXT("ImportDeleteFailed", "Failed to delete '{0}'. The asset is referenced by other content."), FText::FromString(PackageName));
+								FMessageDialog::Open(EAppMsgType::Ok, Message);
+								UE_LOG(LogAssetTools, Warning, TEXT("%s"), *Message.ToString());
+								OnNewImportRecord(ImportAssetType, FileExtension, bImportSucceeded, bImportWasCancelled, ImportStartTime);
+								continue;
+							}
+							else
+							{
+								// succeed, recreate package since it has been deleted
+								Pkg = CreatePackage(nullptr, *PackageName);
+							}
+						}
+						else
+						{
+							// User chose not to replace the package
+							bImportWasCancelled = true;
+							OnNewImportRecord(ImportAssetType, FileExtension, bImportSucceeded, bImportWasCancelled, ImportStartTime);
+							continue;
+						}
+					}
+				}
+			}
+
+			// Check for a package that was marked for delete in source control
+			if(!CheckForDeletedPackage(Pkg))
+			{
+				OnNewImportRecord(ImportAssetType, FileExtension, bImportSucceeded, bImportWasCancelled, ImportStartTime);
+				continue;
+			}
+
+			Factory->SetAutomatedAssetImportData(Params.ImportData);
+
+			ImportAssetType = Factory->ResolveSupportedClass();
+			UObject* Result = Factory->ImportObject(ImportAssetType, Pkg, FName(*Name), RF_Public | RF_Standalone, Filename, nullptr, bImportWasCancelled);
+
+			Factory->SetAutomatedAssetImportData(nullptr);
+
+			// Do not report any error if the operation was canceled.
+			if(!bImportWasCancelled)
+			{
+				if(Result)
+				{
+					ReturnObjects.Add(Result);
+
+					// Notify the asset registry
+					FAssetRegistryModule::AssetCreated(Result);
+					GEditor->BroadcastObjectReimported(Result);
+
+					bImportSucceeded = true;
+				}
+				else
+				{
+					const FText Message = FText::Format(LOCTEXT("ImportFailed_Generic", "Failed to import '{0}'. Failed to create asset '{1}'.\nPlease see Output Log for details."), FText::FromString(Filename), FText::FromString(PackageName));
+					if(!bAutomatedImport)
+					{
+						FMessageDialog::Open(EAppMsgType::Ok, Message);
+					}
+					UE_LOG(LogAssetTools, Warning, TEXT("%s"), *Message.ToString());
+				}
+			}
+
+			// Refresh the supported class.  Some factories (e.g. FBX) only resolve their type after reading the file
+			ImportAssetType = Factory->ResolveSupportedClass();
+			OnNewImportRecord(ImportAssetType, FileExtension, bImportSucceeded, bImportWasCancelled, ImportStartTime);
+		}
+		else
+		{
+			// A factory or extension was not found. The extension warning is above. If a factory was not found, the user likely canceled a factory configuration dialog.
+		}
+	}
+
+	SlowTask.EnterProgressFrame(1);
+
+	// Clean up and remove the factories we created from the root set
+	for(auto ExtensionIt = ExtensionToFactoriesMap.CreateConstIterator(); ExtensionIt; ++ExtensionIt)
+	{
+		for(auto FactoryIt = ExtensionIt.Value().CreateConstIterator(); FactoryIt; ++FactoryIt)
+		{
+			(*FactoryIt)->CleanUp();
+			(*FactoryIt)->RemoveFromRoot();
+		}
+	}
+
+	// Sync content browser to the newly created assets
+	if(ReturnObjects.Num() && (bSyncToBrowser != false))
+	{
+		FAssetTools::Get().SyncBrowserToAssets(ReturnObjects);
+	}
+
+	return ReturnObjects;
+}
+
 FAssetTools& FAssetTools::Get()
diff --git a/Engine/Source/Developer/AssetTools/Private/AssetTools.h b/Engine/Source/Developer/AssetTools/Private/AssetTools.h
--- a/Engine/Source/Developer/AssetTools/Private/AssetTools.h
+++ b/Engine/Source/Developer/AssetTools/Private/AssetTools.h
@@ -2,6 +2,29 @@
 
 #pragma once
 
+/** Parameters for importing specific set of files */
+struct FAssetImportParams
+{
+	FAssetImportParams()
+		: SpecifiedFactory(nullptr)
+		, ImportData(nullptr)
+		, bSyncToBrowser(true)
+		, bForceOverrideExisting(false)
+		, bAutomated(false)
+	{}
+
+	/** Factory to use for importing files */
+	UFactory* SpecifiedFactory;
+	/** Data used to determine rules for importing assets through the automated command line interface */
+	const UAutomatedAssetImportData* ImportData;
+	/** Whether or not to sync the content browser to the assets after import */
+	bool bSyncToBrowser : 1;
+	/** Whether or not we are forcing existing assets to be overriden without asking */
+	bool bForceOverrideExisting : 1;
+	/** Whether or not this is an automated import */
+	bool bAutomated : 1;
+};
+
 class FAssetTools : public IAssetTools
 {
 public:
@@ -29,6 +52,7 @@ public:
 	virtual void RenameAssets(const TArray<FAssetRenameData>& AssetsAndNames) const override;
 	virtual TArray<UObject*> ImportAssets(const FString& DestinationPath) override;
 	virtual TArray<UObject*> ImportAssets(const TArray<FString>& Files, const FString& DestinationPath, UFactory* ChosenFactory, bool bSyncToBrowser = true) const override;
+	virtual TArray<UObject*> ImportAssetsAutomated(const UAutomatedAssetImportData& ImportData) const override;
 	virtual void CreateUniqueAssetName(const FString& InBasePackageName, const FString& InSuffix, FString& OutPackageName, FString& OutAssetName) const override;
 	virtual bool AssetUsesGenericThumbnail( const FAssetData& AssetData ) const override;
 	virtual void DiffAgainstDepot(UObject* InObject, const FString& InPackagePath, const FString& InPackageName) const override;
@@ -69,6 +93,9 @@ private:
 	/** Records what assets users are creating */
 	static void OnNewCreateRecord(UClass* AssetType, bool bDuplicated);
 
+	/** Internal method that performs the actual asset importing */
+	TArray<UObject*> ImportAssetsInternal(const TArray<FString>& Files, const FString& RootDestinationPath, const FAssetImportParams& ImportParams) const;
+
 private:
 	/** The manager to handle renaming assets */
 	TSharedRef<FAssetRenameManager> AssetRenameManager;
diff --git a/Engine/Source/Developer/AssetTools/Public/IAssetTools.h b/Engine/Source/Developer/AssetTools/Public/IAssetTools.h
--- a/Engine/Source/Developer/AssetTools/Public/IAssetTools.h
+++ b/Engine/Source/Developer/AssetTools/Public/IAssetTools.h
@@ -130,7 +130,12 @@ public:
 	/** Event issued at the end of the rename process */
 	virtual FAssetPostRenameEvent& OnAssetPostRename() = 0;
 
-	/** Opens a file open dialog to choose files to import to the destination path. */
+	/** 
+	 * Opens a file open dialog to choose files to import to the destination path.
+	 * 
+	 * @param DestinationPath	Path to import files to
+	 * @return list of sucessfully imported assets
+	 */
 	virtual TArray<UObject*> ImportAssets(const FString& DestinationPath) = 0;
 
 	/** 
@@ -140,9 +145,18 @@ public:
 	 * @param DestinationPath	destination path for imported files
 	 * @param ChosenFactory		Specific factory to use for object creation
 	 * @param bSyncToBrowser	If true sync content browser to first imported asset after import
+	 * @return list of sucessfully imported assets
 	 */
 	virtual TArray<UObject*> ImportAssets(const TArray<FString>& Files, const FString& DestinationPath, UFactory* ChosenFactory = NULL, bool bSyncToBrowser = true) const = 0;
 
+	/**
+	 * Imports assets using data specified completely up front.  Does not ever ask any questions of the user or show any modal error messages
+	 *
+	 * @param AutomatedImportData	Data that specifies how to import each file
+	 * @return list of sucessfully imported assets
+	 */
+	virtual TArray<UObject*> ImportAssetsAutomated(const class UAutomatedAssetImportData& ImportData) const = 0;
+
 	/** Creates a unique package and asset name taking the form InBasePackageName+InSuffix */
 	virtual void CreateUniqueAssetName(const FString& InBasePackageName, const FString& InSuffix, FString& OutPackageName, FString& OutAssetName) const = 0;
 
diff --git a/Engine/Source/Developer/AutomationController/Private/AutomationCommandline.cpp b/Engine/Source/Developer/AutomationController/Private/AutomationCommandline.cpp
--- a/Engine/Source/Developer/AutomationController/Private/AutomationCommandline.cpp
+++ b/Engine/Source/Developer/AutomationController/Private/AutomationCommandline.cpp
@@ -120,12 +120,13 @@ public:
 			Filters[FilterIndex] = Filters[FilterIndex].Replace(TEXT(" "), TEXT(""));
 		}
 
-		for ( int FilterIndex = 0; FilterIndex < Filters.Num(); ++FilterIndex )
+		for ( int32 TestIndex = 0; TestIndex < AllTestNames.Num(); ++TestIndex )
 		{
-			for (int32 TestIndex = 0; TestIndex < AllTestNames.Num(); ++TestIndex)
+			AllTestNamesNoWhiteSpaces = AllTestNames[TestIndex];
+			AllTestNamesNoWhiteSpaces = AllTestNamesNoWhiteSpaces.Replace(TEXT(" "), TEXT(""));
+
+			for ( int32 FilterIndex = 0; FilterIndex < Filters.Num(); ++FilterIndex )
 			{
-				AllTestNamesNoWhiteSpaces = AllTestNames[TestIndex];
-				AllTestNamesNoWhiteSpaces = AllTestNamesNoWhiteSpaces.Replace(TEXT(" "), TEXT(""));
 				if (AllTestNamesNoWhiteSpaces.Contains(Filters[FilterIndex]))
 				{
 					OutTestNames.Add(AllTestNames[TestIndex]);
diff --git a/Engine/Source/Developer/AutomationController/Private/AutomationControllerManager.h b/Engine/Source/Developer/AutomationController/Private/AutomationControllerManager.h
--- a/Engine/Source/Developer/AutomationController/Private/AutomationControllerManager.h
+++ b/Engine/Source/Developer/AutomationController/Private/AutomationControllerManager.h
@@ -4,53 +4,65 @@
 #include "AutomationControllerManager.generated.h"
 
 USTRUCT()
-	struct FAutomatedTestResult
-	{
-		GENERATED_BODY()
-	public:
-
-		UPROPERTY()
-		FString TestName;
-		UPROPERTY()
-		FString TestResult;
-		UPROPERTY()
-		TArray<FString> TestInfo;
-		UPROPERTY()
-		TArray<FString> TestErrors;
-		UPROPERTY()
-		TArray<FString> TestWarnings;
-		FAutomatedTestResult()
-		{
-			TestResult = TEXT("Not Run");
-		}
-	};	
-	USTRUCT()
-	struct FAutomatedTestPassResults
+struct FAutomatedTestResult
+{
+	GENERATED_BODY()
+public:
+
+	UPROPERTY()
+	FString TestName;
+
+	UPROPERTY()
+	FString TestResult;
+
+	UPROPERTY()
+	TArray<FString> TestInfo;
+
+	UPROPERTY()
+	TArray<FString> TestErrors;
+
+	UPROPERTY()
+	TArray<FString> TestWarnings;
+
+	FAutomatedTestResult()
 	{
-		GENERATED_BODY()
+		TestResult = TEXT("Not Run");
+	}
+};
 
-	public:
-		FAutomatedTestPassResults()
-		{
+USTRUCT()
+struct FAutomatedTestPassResults
+{
+	GENERATED_BODY()
 
+public:
+	FAutomatedTestPassResults()
+		: NumSucceeded(0)
+		, NumFailed(0)
+		, NumNotRun(0)
+	{
+	}
 
-		}
-		UPROPERTY()
-			int NumSucceeded;
-		UPROPERTY()
-			int NumFailed;
-		UPROPERTY()
-			int NumNotRun;
-		UPROPERTY()
-			TArray<FAutomatedTestResult> TestInformation;
-		void ClearAllEntries()
-		{
-			NumSucceeded = 0;
-			NumFailed = 0;
-			NumNotRun = 0;
-			TestInformation.Empty();
-		}
-	};
+	UPROPERTY()
+	int32 NumSucceeded;
+
+	UPROPERTY()
+	int32 NumFailed;
+
+	UPROPERTY()
+	int32 NumNotRun;
+
+	UPROPERTY()
+	TArray<FAutomatedTestResult> TestInformation;
+
+	void ClearAllEntries()
+	{
+		NumSucceeded = 0;
+		NumFailed = 0;
+		NumNotRun = 0;
+		TestInformation.Empty();
+	}
+};
 
 
 
@@ -59,8 +71,6 @@ USTRUCT()
  */
 class FAutomationControllerManager : public IAutomationControllerManager
 {
-
-
 public:
 	// IAutomationController Interface
 
diff --git a/Engine/Source/Developer/AutomationController/Private/AutomationControllerManger.cpp b/Engine/Source/Developer/AutomationController/Private/AutomationControllerManger.cpp
--- a/Engine/Source/Developer/AutomationController/Private/AutomationControllerManger.cpp
+++ b/Engine/Source/Developer/AutomationController/Private/AutomationControllerManger.cpp
@@ -767,7 +767,7 @@ void FAutomationControllerManager::HandleRequestNextNetworkCommandMessage( const
 
 void FAutomationControllerManager::HandleRequestTestsReplyMessage(const FAutomationWorkerRequestTestsReply& Message, const IMessageContextRef& Context)
 {
-	FAutomationTestInfo NewTest(Message.TestInfo);
+	FAutomationTestInfo NewTest = Message.GetTestInfo();
 	TestInfo.Add(NewTest);
 }
 
diff --git a/Engine/Source/Developer/AutomationController/Private/AutomationReport.cpp b/Engine/Source/Developer/AutomationController/Private/AutomationReport.cpp
--- a/Engine/Source/Developer/AutomationController/Private/AutomationReport.cpp
+++ b/Engine/Source/Developer/AutomationController/Private/AutomationReport.cpp
@@ -58,6 +58,11 @@ const FString& FAutomationReport::GetDisplayName() const
 	return TestInfo.GetDisplayName();
 }
 
+const FString& FAutomationReport::GetFullTestPath() const
+{
+	return TestInfo.GetFullTestPath();
+}
+
 FString FAutomationReport::GetDisplayNameWithDecoration() const
 {
 	FString FinalDisplayName = TestInfo.GetDisplayName();
@@ -769,7 +774,7 @@ TSharedPtr<IAutomationReport> FAutomationReport::EnsureReportExists(FAutomationT
 		else
 		{
 			// Create a parent node
-			FAutomationTestInfo ParentTestInfo(NameToMatch, "", InTestInfo.GetTestFlags(), InTestInfo.GetNumParticipantsRequired());
+			FAutomationTestInfo ParentTestInfo(NameToMatch, TEXT(""), TEXT(""), InTestInfo.GetTestFlags(), InTestInfo.GetNumParticipantsRequired());
 			MatchTest = MakeShareable(new FAutomationReport(ParentTestInfo, true));
 		}
 		//make new test
diff --git a/Engine/Source/Developer/AutomationController/Private/AutomationReport.h b/Engine/Source/Developer/AutomationController/Private/AutomationReport.h
--- a/Engine/Source/Developer/AutomationController/Private/AutomationReport.h
+++ b/Engine/Source/Developer/AutomationController/Private/AutomationReport.h
@@ -22,6 +22,7 @@ public:
 	virtual FString GetOpenCommand() const override;
 	virtual FString GetCommand() const override;
 	virtual const FString& GetDisplayName() const override;
+	virtual const FString& GetFullTestPath() const override;
 	virtual FString GetDisplayNameWithDecoration() const override;
 	virtual int32 GetTotalNumChildren() const override;
 	virtual int32 GetTotalNumFilteredChildren() const override;
diff --git a/Engine/Source/Developer/AutomationController/Public/Interfaces/IAutomationReport.h b/Engine/Source/Developer/AutomationController/Public/Interfaces/IAutomationReport.h
--- a/Engine/Source/Developer/AutomationController/Public/Interfaces/IAutomationReport.h
+++ b/Engine/Source/Developer/AutomationController/Public/Interfaces/IAutomationReport.h
@@ -183,6 +183,8 @@ public:
 	 */
 	virtual const FString& GetDisplayName() const = 0;
 
+	virtual const FString& GetFullTestPath() const = 0;
+
 	/**
 	 * Returns the name of this level in the test hierarchy for the purposes of UI.
 	 * (Editor.Maps.LoadAll.parameters) would have 3 internal tree nodes and a variable number of leaf nodes depending on how many maps existed
diff --git a/Engine/Source/Developer/AutomationWindow/Private/SAutomationWindow.cpp b/Engine/Source/Developer/AutomationWindow/Private/SAutomationWindow.cpp
--- a/Engine/Source/Developer/AutomationWindow/Private/SAutomationWindow.cpp
+++ b/Engine/Source/Developer/AutomationWindow/Private/SAutomationWindow.cpp
@@ -1528,6 +1528,7 @@ void SAutomationWindow::PopulateReportSearchStrings( const TSharedPtr< IAutomati
 	}
 
 	OutSearchStrings.Add( Report->GetDisplayName() );
+	OutSearchStrings.Add( Report->GetFullTestPath() );
 }
 
 
diff --git a/Engine/Source/Developer/BlueprintCompilerCppBackend/Private/BPCompilerTests.cpp b/Engine/Source/Developer/BlueprintCompilerCppBackend/Private/BPCompilerTests.cpp
--- a/Engine/Source/Developer/BlueprintCompilerCppBackend/Private/BPCompilerTests.cpp
+++ b/Engine/Source/Developer/BlueprintCompilerCppBackend/Private/BPCompilerTests.cpp
@@ -72,9 +72,9 @@ private:
 // Helper functions introduced to load classes (generated or native:
 static UClass* GetGeneratedClass(const TCHAR* TestFolder, const TCHAR* ClassName, FAutomationTestBase* Context, FOwnedObjectsHelper &OwnedObjects)
 {
-	FString FullName = FString::Printf(TEXT("/Engine/NotForLicensees/Automation/CompilerTests/%s/%s.%s"), TestFolder, ClassName, ClassName);
-	UBlueprint* Blueprint = LoadObject<UBlueprint>(NULL, *FullName);
-
+	FString FullName = FString::Printf(TEXT("/RuntimeTests/CompilerTests/%s/%s.%s"), TestFolder, ClassName, ClassName);
+	UBlueprint* Blueprint = LoadObject<UBlueprint>(nullptr, *FullName);
+	
 	if (!Blueprint)
 	{
 		Context->AddWarning(FString::Printf(TEXT("Missing blueprint for test: '%s'"), *FullName));
@@ -97,8 +97,8 @@ static UClass* GetNativeClass(const TCHAR* TestFolder, const TCHAR* ClassName, F
 {
 	CollectGarbage(RF_NoFlags);
 
-	FString FullName = FString::Printf(TEXT("/Engine/NotForLicensees/Automation/CompilerTests/%s/%s"), TestFolder, ClassName);
-	UPackage* NativePackage = CreatePackage(NULL, *FullName);
+	FString FullName = FString::Printf(TEXT("/RuntimeTests/CompilerTests/%s/%s"), TestFolder, ClassName);
+	UPackage* NativePackage = CreatePackage(nullptr, *FullName);
 	check(NativePackage);
 
 	const FString FStringFullPathName = FString::Printf(TEXT("%s.%s_C"), *FullName, ClassName);
@@ -190,7 +190,7 @@ static void Call(UObject* Target, const TCHAR* FunctionName, void* Args = nullpt
 }
 
 // Remove EAutomationTestFlags::Disabled to enable these tests, note that these will need to be moved into the ClientContext because we can only test cooked content:
-static const uint32 CompilerTestFlags = EAutomationTestFlags::EditorContext | EAutomationTestFlags::EngineFilter;
+static const uint32 CompilerTestFlags = EAutomationTestFlags::ClientContext | EAutomationTestFlags::EngineFilter;
 
 // Tests:
 IMPLEMENT_SIMPLE_AUTOMATION_TEST(FBPCompilerArrayTest, "Project.Blueprints.NativeBackend.ArrayTest", CompilerTestFlags)
diff --git a/Engine/Source/Developer/FunctionalTesting/Classes/FunctionalTest.h b/Engine/Source/Developer/FunctionalTesting/Classes/FunctionalTest.h
--- a/Engine/Source/Developer/FunctionalTesting/Classes/FunctionalTest.h
+++ b/Engine/Source/Developer/FunctionalTesting/Classes/FunctionalTest.h
@@ -176,7 +176,7 @@ public:
 	void WriteLogFile(const FString& CaptureDir, const FString& CaptureExtension);
 	/** Returns true if this stats tracker is currently recording performance stats. */
 	UFUNCTION(BlueprintCallable, Category = Perf)
-	bool IsRecording()const;
+	bool IsRecording() const;
 
 	/** Does any init work across all tests.. */
 	UFUNCTION(BlueprintCallable, Category = Perf)
@@ -185,15 +185,15 @@ public:
 	UFUNCTION(BlueprintCallable, Category = Perf)
 	void OnAllTestsComplete();
 
-	const FPerfStatsRecord* GetCurrentRecord()const;
+	const FPerfStatsRecord* GetCurrentRecord() const;
 	FPerfStatsRecord* GetCurrentRecord();
 
 	UFUNCTION(BlueprintCallable, Category = Perf)
-	bool IsCurrentRecordWithinGPUBudget()const;
+	bool IsCurrentRecordWithinGPUBudget() const;
 	UFUNCTION(BlueprintCallable, Category = Perf)
-	bool IsCurrentRecordWithinGameThreadBudget()const;
+	bool IsCurrentRecordWithinGameThreadBudget() const;
 	UFUNCTION(BlueprintCallable, Category = Perf)
-	bool IsCurrentRecordWithinRenderThreadBudget()const;
+	bool IsCurrentRecordWithinRenderThreadBudget() const;
 	//End basic stats recording.
 
 	// Automatic traces capturing 
@@ -222,7 +222,7 @@ public:
 	FString StartOfTestingTime;
 };
 
-UENUM()
+UENUM(BlueprintType)
 enum class EFunctionalTestResult : uint8
 {
 	/**
@@ -246,7 +246,7 @@ class FUNCTIONALTESTING_API AFunctionalTest : public AActor
 	GENERATED_BODY()
 
 public:
-	AFunctionalTest(const FObjectInitializer& ObjectInitializer);
+	AFunctionalTest(const FObjectInitializer& ObjectInitializer = FObjectInitializer::Get());
 
 private_subobject:
 	DEPRECATED_FORGAME(4.6, "SpriteComponent should not be accessed directly, please use GetSpriteComponent() function instead. SpriteComponent will soon be private and your code will not compile.")
@@ -257,6 +257,9 @@ protected:
 	UPROPERTY(EditAnywhere, BlueprintReadOnly, Category="Functional Testing")
 	uint32 bIsEnabled:1;
 
+	UPROPERTY(EditAnywhere, BlueprintReadOnly, Category="Functional Testing")
+	uint32 bWarningsAsErrors:1;
+
 	UPROPERTY(EditAnywhere, BlueprintReadOnly, Category="Functional Testing", meta=( MultiLine="true" ))
 	FString Description;
 
@@ -421,19 +424,23 @@ public:
 	UFUNCTION(BlueprintCallable, Category = "Asserts", DisplayName = "Assert Not Equal (String)", meta = ( HidePin = "ContextObject", DefaultToSelf = "ContextObject"))
 	void AssertNotEqual_String(const FString Actual, const FString NotExpected, const FString& What, const UObject* ContextObject = nullptr);
 
+	UFUNCTION(BlueprintCallable, Category = "Reporting")
+	void AddWarning(const FString Message);
+
+	UFUNCTION(BlueprintCallable, Category = "Reporting")
+	void AddError(const FString Message);
+
 protected:
-	void LogAssertFail(const FString& Message);
-	void LogAssertPass(const FString& Message);
+	void LogStep(ELogVerbosity::Type Verbosity, const FString& Message);
 
 public:
 	virtual bool RunTest(const TArray<FString>& Params = TArray<FString>());
 
 public:
-	FString	GetCurrentStepName() { return IsInStep() ? CurrentStepName : FString(); }
-	void 	StartStep(const FString& StepName) { CurrentStepName = StepName; bInStep = true; }
-	void 	FinishStep() { CurrentStepName.Empty(); bInStep = false; }
-	bool	IsInStep() { return bInStep; }
-
+	FString	GetCurrentStepName() const;
+	void 	StartStep(const FString& StepName);
+	void 	FinishStep();
+	bool	IsInStep() const;
 
 	UFUNCTION(BlueprintCallable, Category="Functional Testing")
 	virtual void FinishTest(EFunctionalTestResult TestResult, const FString& Message);
@@ -552,8 +559,7 @@ public:
 protected:
 	bool bIsRunning;
 
-	bool bInStep;
-	FString CurrentStepName;
+	TArray<FString> Steps;
 	
 	float TotalTime;
 
diff --git a/Engine/Source/Developer/FunctionalTesting/Classes/ScreenshotFunctionalTest.h b/Engine/Source/Developer/FunctionalTesting/Classes/ScreenshotFunctionalTest.h
--- a/Engine/Source/Developer/FunctionalTesting/Classes/ScreenshotFunctionalTest.h
+++ b/Engine/Source/Developer/FunctionalTesting/Classes/ScreenshotFunctionalTest.h
@@ -3,65 +3,10 @@
 #pragma once
 
 #include "FunctionalTest.h"
+#include "AutomationScreenshotOptions.h"
 
 #include "ScreenshotFunctionalTest.generated.h"
 
-UENUM()
-enum class EComparisonTolerance : uint8
-{
-	Zero,
-	Low,
-	High,
-	Custom
-};
-
-USTRUCT()
-struct FComparisonToleranceAmount
-{
-	GENERATED_BODY()
-
-public:
-
-	FComparisonToleranceAmount()
-		: Red(0)
-		, Green(0)
-		, Blue(0)
-		, Alpha(0)
-		, MinBrightness(0)
-		, MaxBrightness(255)
-	{
-	}
-
-	FComparisonToleranceAmount(uint8 R, uint8 G, uint8 B, uint8 A, uint8 InMinBrightness, uint8 InMaxBrightness)
-		: Red(R)
-		, Green(G)
-		, Blue(B)
-		, Alpha(A)
-		, MinBrightness(InMinBrightness)
-		, MaxBrightness(InMaxBrightness)
-	{
-	}
-
-public:
-	UPROPERTY(EditAnywhere, Category="Tolerance")
-	uint8 Red;
-	
-	UPROPERTY(EditAnywhere, Category="Tolerance")
-	uint8 Green;
-	
-	UPROPERTY(EditAnywhere, Category="Tolerance")
-	uint8 Blue;
-	
-	UPROPERTY(EditAnywhere, Category="Tolerance")
-	uint8 Alpha;
-	
-	UPROPERTY(EditAnywhere, Category="Tolerance")
-	uint8 MinBrightness;
-
-	UPROPERTY(EditAnywhere, Category="Tolerance")
-	uint8 MaxBrightness;
-};
-
 /**
  * 
  */
@@ -91,54 +36,6 @@ protected:
 	UPROPERTY()
 	class UCameraComponent* ScreenshotCamera;
 
-	/**
-	 * The desired resolution of the screenshot, if none is provided, it will use the default for the
-	 * platform setup in the automation settings.
-	 */
-	UPROPERTY(EditAnywhere, Category="Screenshot")
-	FIntPoint Resolution;
-
-	/**
-	 * Allows you to screenshot a buffer other than the default final lit scene image.  Useful if you're
-	 * trying to build a test for a specific GBuffer, that may be harder to tell if errors are introduced 
-	 * in it.
-	 */
-	UPROPERTY(EditAnywhere, Category="Screenshot")
-	FName VisualizeBuffer;
-
-	/**
-	 * These are quick defaults for tolerance levels, we default to low, because generally there's some
-	 * constant variability in every pixel's color introduced by TxAA.
-	 */
-	UPROPERTY(EditAnywhere, Category="Comparison")
-	EComparisonTolerance Tolerance;
-
-	/**
-	 * For each channel and brightness levels you can control a region where the colors are found to be
-	 * essentially the same.  Generally this is necessary as modern rendering techniques tend to introduce
-	 * noise constantly to hide aliasing.
-	 */
-	UPROPERTY(EditAnywhere, Category="Comparison")
-	FComparisonToleranceAmount ToleranceAmount;
-
-	/**
-	 * After you've accounted for color tolerance changes, you now need to control for total acceptable error.
-	 * Which depending on how pixels were colored on triangle edges may be a few percent of the image being 
-	 * outside the tolerance levels.
-	 */
-	UPROPERTY(EditAnywhere, Category="Comparison", meta=( ClampMin = 0, ClampMax = 1, UIMin = 0, UIMax = 1 ))
-	float MaximumAllowedError;
-
-	/**
-	 * If this is true, we search neighboring pixels looking for the expected pixel as what may have happened, is
-	 * that the pixel shifted a little.
-	 */
-	UPROPERTY(EditAnywhere, Category="Comparison")
-	bool bIgnoreAntiAliasing;
-
-	/**
-	 * If this is true, all we compare is luminance of the scene.
-	 */
-	UPROPERTY(EditAnywhere, Category="Comparison", AdvancedDisplay)
-	bool bIgnoreColors;
+	UPROPERTY(EditAnywhere, Category="Screenshot", SimpleDisplay)
+	FAutomationScreenshotOptions ScreenshotOptions;
 };
\ No newline at end of file
diff --git a/Engine/Source/Developer/FunctionalTesting/Private/AutomationBlueprintFunctionLibrary.cpp b/Engine/Source/Developer/FunctionalTesting/Private/AutomationBlueprintFunctionLibrary.cpp
--- a/Engine/Source/Developer/FunctionalTesting/Private/AutomationBlueprintFunctionLibrary.cpp
+++ b/Engine/Source/Developer/FunctionalTesting/Private/AutomationBlueprintFunctionLibrary.cpp
@@ -14,6 +14,8 @@
 
 #include "AutomationBlueprintFunctionLibrary.h"
 
+#define LOCTEXT_NAMESPACE "Automation"
+
 DEFINE_LOG_CATEGORY_STATIC(BlueprintAssertion, Error, Error)
 DEFINE_LOG_CATEGORY_STATIC(AutomationFunctionLibrary, Log, Log)
 
@@ -29,25 +31,110 @@ static TAutoConsoleVariable<int32> CVarAutomationScreenshotResolutionHeight(
 	TEXT("The height of automation screenshots."),
 	ECVF_Default);
 
+FAutomationScreenshotOptions::FAutomationScreenshotOptions()
+	: Resolution(ForceInit)
+	, Delay(0.2f)
+	, bDisableNoisyRenderingFeatures(false)
+	, VisualizeBuffer(NAME_None)
+	, Tolerance(EComparisonTolerance::Low)
+	, ToleranceAmount()
+	, MaximumAllowedError(0.02f)
+	, bIgnoreAntiAliasing(true)
+	, bIgnoreColors(false)
+{
+}
+
+void FAutomationScreenshotOptions::SetToleranceAmounts(EComparisonTolerance InTolerance)
+{
+	switch ( InTolerance )
+	{
+	case EComparisonTolerance::Zero:
+		ToleranceAmount = FComparisonToleranceAmount(0, 0, 0, 0, 0, 255);
+		break;
+	case EComparisonTolerance::Low:
+		ToleranceAmount = FComparisonToleranceAmount(16, 16, 16, 16, 16, 240);
+		break;
+	case EComparisonTolerance::Medium:
+		ToleranceAmount = FComparisonToleranceAmount(24, 24, 24, 24, 24, 220);
+		break;
+	case EComparisonTolerance::High:
+		ToleranceAmount = FComparisonToleranceAmount(32, 32, 32, 32, 64, 96);
+		break;
+	}
+}
+
 #if (WITH_DEV_AUTOMATION_TESTS || WITH_PERF_AUTOMATION_TESTS)
 
 class FAutomationScreenshotTaker
 {
-	FOnScreenshotCaptured&	ScreenshotCapturedDelegate;
 	FDelegateHandle			DelegateHandle;
 	FString					Name;
+	FAutomationScreenshotOptions Options;
+
+	int32 OriginalPostProcessing;
+	int32 OriginalMotionBlur;
+	int32 OriginalSSRQuality;
+	int32 OriginalEyeAdaptation;
+	int32 OriginalContactShadows;
 
 public:
-	FAutomationScreenshotTaker(FOnScreenshotCaptured& InScreenshotCapturedDelegate, const FString& InName)
-		: ScreenshotCapturedDelegate(InScreenshotCapturedDelegate)
-		, DelegateHandle()
+	FAutomationScreenshotTaker(const FString& InName, FAutomationScreenshotOptions InOptions)
+		: DelegateHandle()
 		, Name(InName)
+		, Options(InOptions)
 	{
-		DelegateHandle = ScreenshotCapturedDelegate.AddRaw(this, &FAutomationScreenshotTaker::GrabScreenShot);
+		DelegateHandle = GEngine->GameViewport->OnScreenshotCaptured().AddRaw(this, &FAutomationScreenshotTaker::GrabScreenShot);
+
+		check(IsInGameThread());
+
+		IConsoleVariable* PostProcessingQuality = IConsoleManager::Get().FindConsoleVariable(TEXT("r.PostProcessAAQuality"));
+		IConsoleVariable* MotionBlur = IConsoleManager::Get().FindConsoleVariable(TEXT("r.MotionBlurQuality"));
+		IConsoleVariable* ScreenSpaceReflections = IConsoleManager::Get().FindConsoleVariable(TEXT("r.SSR.Quality"));
+		IConsoleVariable* EyeAdapation = IConsoleManager::Get().FindConsoleVariable(TEXT("r.EyeAdaptationQuality"));
+		IConsoleVariable* ContactShadows = IConsoleManager::Get().FindConsoleVariable(TEXT("r.ContactShadows"));
+
+		OriginalPostProcessing = PostProcessingQuality->GetInt();
+		OriginalMotionBlur = MotionBlur->GetInt();
+		OriginalSSRQuality = ScreenSpaceReflections->GetInt();
+		OriginalEyeAdaptation = EyeAdapation->GetInt();
+		OriginalContactShadows = ContactShadows->GetInt();
+
+		if ( Options.bDisableNoisyRenderingFeatures )
+		{
+			PostProcessingQuality->Set(1);
+			MotionBlur->Set(0);
+			ScreenSpaceReflections->Set(0);
+			EyeAdapation->Set(0);
+			ContactShadows->Set(0);
+		}
+	}
+
+	~FAutomationScreenshotTaker()
+	{
+		check(IsInGameThread());
+
+		if ( Options.bDisableNoisyRenderingFeatures )
+		{
+			IConsoleVariable* PostProcessingQuality = IConsoleManager::Get().FindConsoleVariable(TEXT("r.PostProcessAAQuality"));
+			IConsoleVariable* MotionBlur = IConsoleManager::Get().FindConsoleVariable(TEXT("r.MotionBlurQuality"));
+			IConsoleVariable* ScreenSpaceReflections = IConsoleManager::Get().FindConsoleVariable(TEXT("r.SSR.Quality"));
+			IConsoleVariable* EyeAdapation = IConsoleManager::Get().FindConsoleVariable(TEXT("r.EyeAdaptationQuality"));
+			IConsoleVariable* ContactShadows = IConsoleManager::Get().FindConsoleVariable(TEXT("r.ContactShadows"));
+
+			PostProcessingQuality->Set(OriginalPostProcessing);
+			MotionBlur->Set(OriginalMotionBlur);
+			ScreenSpaceReflections->Set(OriginalSSRQuality);
+			EyeAdapation->Set(OriginalEyeAdaptation);
+			ContactShadows->Set(OriginalContactShadows);
+		}
+
+		GEngine->GameViewport->OnScreenshotCaptured().Remove(DelegateHandle);
 	}
 
 	void GrabScreenShot(int32 InSizeX, int32 InSizeY, const TArray<FColor>& InImageData)
 	{
+		check(IsInGameThread());
+
 		FAutomationScreenshotData Data = AutomationCommon::BuildScreenshotData(GWorld->GetName(), Name, InSizeX, InSizeY);
 		//Data.bHasComparisonRules = true;
 
@@ -55,20 +142,25 @@ public:
 
 		UE_LOG(AutomationFunctionLibrary, Log, TEXT("Screenshot captured as %s"), *Data.Path);
 
-		ScreenshotCapturedDelegate.Remove(DelegateHandle);
-
 		delete this;
 	}
 };
 
 #endif
 
-UAutomationBlueprintFunctionLibrary::UAutomationBlueprintFunctionLibrary(const class FObjectInitializer& Initializer) : Super(Initializer)
+UAutomationBlueprintFunctionLibrary::UAutomationBlueprintFunctionLibrary(const class FObjectInitializer& Initializer)
+	: Super(Initializer)
 {
 }
 
-bool UAutomationBlueprintFunctionLibrary::TakeAutomationScreenshotInternal(const FString& Name, FIntPoint Resolution)
+bool UAutomationBlueprintFunctionLibrary::TakeAutomationScreenshotInternal(const FString& Name, FAutomationScreenshotOptions Options)
 {
+	if ( !FAutomationTestFramework::Get().IsScreenshotAllowed() )
+	{
+		UE_LOG(AutomationFunctionLibrary, Log, TEXT("Attempted to capture screenshot (%s) but screenshots are not enabled."), *Name);
+		return false;
+	}
+
 	// Fallback resolution if all else fails for screenshots.
 	uint32 ResolutionX = 1280;
 	uint32 ResolutionY = 720;
@@ -82,10 +174,10 @@ bool UAutomationBlueprintFunctionLibrary::TakeAutomationScreenshotInternal(const
 	}
 	
 	// If there's an override resolution, use that instead.
-	if ( Resolution.GetMin() > 0 )
+	if ( Options.Resolution.GetMin() > 0 )
 	{
-		ResolutionX = (uint32)Resolution.X;
-		ResolutionY = (uint32)Resolution.Y;
+		ResolutionX = (uint32)Options.Resolution.X;
+		ResolutionY = (uint32)Options.Resolution.Y;
 	}
 	else
 	{
@@ -109,35 +201,35 @@ bool UAutomationBlueprintFunctionLibrary::TakeAutomationScreenshotInternal(const
 	UTexture::ForceUpdateTextureStreaming();
 
 #if (WITH_DEV_AUTOMATION_TESTS || WITH_PERF_AUTOMATION_TESTS)
-	FAutomationScreenshotTaker* TempObject = new FAutomationScreenshotTaker(GEngine->GameViewport->OnScreenshotCaptured(), Name);
+	FAutomationScreenshotTaker* TempObject = new FAutomationScreenshotTaker(Name, Options);
 #endif
 
 	FHighResScreenshotConfig& Config = GetHighResScreenshotConfig();
 	if ( Config.SetResolution(ResolutionX, ResolutionY, 1.0f) )
 	{
-		GEngine->GameViewport->GetGameViewport()->TakeHighResScreenShot();
+		if ( !GEngine->GameViewport->GetGameViewport()->TakeHighResScreenShot() )
+		{
+			// If we failed to take the screenshot, we're going to need to cleanup the automation screenshot taker.
+			delete TempObject;
+			return false;
+		}
+
 		return true;
 	}
 
 	return false;
 }
 
-void UAutomationBlueprintFunctionLibrary::TakeAutomationScreenshot(UObject* WorldContextObject, FLatentActionInfo LatentInfo, const FString& Name, FIntPoint Resolution, float DelayBeforeScreenshotSeconds)
+void UAutomationBlueprintFunctionLibrary::TakeAutomationScreenshot(UObject* WorldContextObject, FLatentActionInfo LatentInfo, const FString& Name, FAutomationScreenshotOptions Options)
 {
-	if (GIsAutomationTesting)
+	if ( GIsAutomationTesting )
 	{
-		if(!FAutomationTestFramework::Get().IsScreenshotAllowed())
-		{
-			UE_LOG(AutomationFunctionLibrary, Error, TEXT("Attempted to capture screenshot (%s) but screenshots are not enabled. If you're seeing this on TeamCity, please make sure the name of your test contains '_VisualTest'"), *Name);
-			return;
-		}
-
-		if (UWorld* World = GEngine->GetWorldFromContextObject(WorldContextObject))
+		if ( UWorld* World = GEngine->GetWorldFromContextObject(WorldContextObject) )
 		{
 			FLatentActionManager& LatentActionManager = World->GetLatentActionManager();
-			if (LatentActionManager.FindExistingAction<FTakeScreenshotAfterTimeLatentAction>(LatentInfo.CallbackTarget, LatentInfo.UUID) == nullptr)
+			if ( LatentActionManager.FindExistingAction<FTakeScreenshotAfterTimeLatentAction>(LatentInfo.CallbackTarget, LatentInfo.UUID) == nullptr )
 			{
-				LatentActionManager.AddNewAction(LatentInfo.CallbackTarget, LatentInfo.UUID, new FTakeScreenshotAfterTimeLatentAction(LatentInfo, Name, Resolution, DelayBeforeScreenshotSeconds));
+				LatentActionManager.AddNewAction(LatentInfo.CallbackTarget, LatentInfo.UUID, new FTakeScreenshotAfterTimeLatentAction(LatentInfo, Name, Options));
 			}
 		}
 	}
@@ -147,11 +239,11 @@ void UAutomationBlueprintFunctionLibrary::TakeAutomationScreenshot(UObject* Worl
 	}
 }
 
-void UAutomationBlueprintFunctionLibrary::TakeAutomationScreenshotAtCamera(UObject* WorldContextObject, FLatentActionInfo LatentInfo, ACameraActor* Camera, const FString& NameOverride, FIntPoint Resolution, float DelayBeforeScreenshotSeconds)
+void UAutomationBlueprintFunctionLibrary::TakeAutomationScreenshotAtCamera(UObject* WorldContextObject, FLatentActionInfo LatentInfo, ACameraActor* Camera, const FString& NameOverride, FAutomationScreenshotOptions Options)
 {
 	if ( Camera == nullptr )
 	{
-		UE_LOG(AutomationFunctionLibrary, Error, TEXT("A camera is required to TakeAutomationScreenshotAtCamera"));
+		FMessageLog("PIE").Error(LOCTEXT("CameraRequired", "A camera is required to TakeAutomationScreenshotAtCamera"));
 		return;
 	}
 
@@ -159,7 +251,7 @@ void UAutomationBlueprintFunctionLibrary::TakeAutomationScreenshotAtCamera(UObje
 
 	if ( PlayerController == nullptr )
 	{
-		UE_LOG(AutomationFunctionLibrary, Error, TEXT("A player controller is required to TakeAutomationScreenshotAtCamera"));
+		FMessageLog("PIE").Error(LOCTEXT("PlayerRequired", "A player controller is required to TakeAutomationScreenshotAtCamera"));
 		return;
 	}
 
@@ -180,17 +272,46 @@ void UAutomationBlueprintFunctionLibrary::TakeAutomationScreenshotAtCamera(UObje
 		FLatentActionManager& LatentActionManager = World->GetLatentActionManager();
 		if ( LatentActionManager.FindExistingAction<FTakeScreenshotAfterTimeLatentAction>(LatentInfo.CallbackTarget, LatentInfo.UUID) == nullptr )
 		{
-			LatentActionManager.AddNewAction(LatentInfo.CallbackTarget, LatentInfo.UUID, new FTakeScreenshotAfterTimeLatentAction(LatentInfo, ScreenshotName, Resolution, DelayBeforeScreenshotSeconds));
+			LatentActionManager.AddNewAction(LatentInfo.CallbackTarget, LatentInfo.UUID, new FTakeScreenshotAfterTimeLatentAction(LatentInfo, ScreenshotName, Options));
 		}
 	}
 }
 
-void UAutomationBlueprintFunctionLibrary::BeginPerformanceCapture()
+//void UAutomationBlueprintFunctionLibrary::BeginPerformanceCapture()
+//{
+//    //::BeginPerformanceCapture();
+//}
+//
+//void UAutomationBlueprintFunctionLibrary::EndPerformanceCapture()
+//{
+//    //::EndPerformanceCapture();
+//}
+
+bool UAutomationBlueprintFunctionLibrary::AreAutomatedTestsRunning()
+{
+	return GIsAutomationTesting;
+}
+
+FAutomationScreenshotOptions UAutomationBlueprintFunctionLibrary::GetDefaultScreenshotOptionsForGameplay(EComparisonTolerance Tolerance)
 {
-    //::BeginPerformanceCapture();
+	FAutomationScreenshotOptions Options;
+	Options.Tolerance = Tolerance;
+	Options.bDisableNoisyRenderingFeatures = true;
+	Options.bIgnoreAntiAliasing = true;
+	Options.SetToleranceAmounts(Tolerance);
+
+	return Options;
 }
 
-void UAutomationBlueprintFunctionLibrary::EndPerformanceCapture()
+FAutomationScreenshotOptions UAutomationBlueprintFunctionLibrary::GetDefaultScreenshotOptionsForRendering(EComparisonTolerance Tolerance)
 {
-    //::EndPerformanceCapture();
+	FAutomationScreenshotOptions Options;
+	Options.Tolerance = Tolerance;
+	Options.bDisableNoisyRenderingFeatures = false;
+	Options.bIgnoreAntiAliasing = true;
+	Options.SetToleranceAmounts(Tolerance);
+
+	return Options;
 }
+
+#undef LOCTEXT_NAMESPACE
\ No newline at end of file
