Copying //UE4/Dev-Framework to //UE4/Dev-Main (Source: //UE4/Dev-Framework @ 3582324)

diff --git a/Engine/Source/Developer/AssetTools/Private/AssetRenameManager.cpp b/Engine/Source/Developer/AssetTools/Private/AssetRenameManager.cpp
--- a/Engine/Source/Developer/AssetTools/Private/AssetRenameManager.cpp
+++ b/Engine/Source/Developer/AssetTools/Private/AssetRenameManager.cpp
@@ -56,7 +56,26 @@ struct FAssetRenameDataWithReferencers : public FAssetRenameData
 		: FAssetRenameData(InRenameData)
 		, bCreateRedirector(false)
 		, bRenameFailed(false)
-	{}
+	{
+		if (Asset.IsValid() && !OldObjectPath.IsValid())
+		{
+			OldObjectPath = FSoftObjectPath(Asset.Get());
+		}
+		else if (OldObjectPath.IsValid() && !Asset.IsValid())
+		{
+			Asset = OldObjectPath.ResolveObject();
+		}
+
+		if (!NewName.IsEmpty() && !NewObjectPath.IsValid())
+		{
+			NewObjectPath = FSoftObjectPath(FString::Printf(TEXT("%s/%s.%s"), *NewPackagePath, *NewName, *NewName));
+		}
+		else if (NewObjectPath.IsValid() && NewName.IsEmpty())
+		{
+			NewName = NewObjectPath.GetAssetName();
+			NewPackagePath = FPackageName::GetLongPackagePath(NewObjectPath.GetLongPackageName());
+		}
+	}
 };
 
 class SRenameFailures : public SCompoundWidget
@@ -71,9 +90,9 @@ public:
 	BEGIN_SLATE_FUNCTION_BUILD_OPTIMIZATION
 	void Construct( const FArguments& InArgs )
 	{
-		for ( auto RenameIt = InArgs._FailedRenames.CreateConstIterator(); RenameIt; ++RenameIt )
+		for (const FText& RenameText : InArgs._FailedRenames)
 		{
-			FailedRenames.Add( MakeShareable( new FText(*RenameIt) ) );
+			FailedRenames.Add( MakeShareable( new FText(RenameText) ) );
 		}
 
 		ChildSlot
@@ -134,7 +153,7 @@ public:
 
 		IMainFrameModule& MainFrameModule = FModuleManager::LoadModuleChecked<IMainFrameModule>(TEXT("MainFrame"));
 
-		if ( MainFrameModule.GetParentWindow().IsValid() )
+		if (MainFrameModule.GetParentWindow().IsValid())
 		{
 			FSlateApplication::Get().AddWindowAsNativeChild(RenameWindow, MainFrameModule.GetParentWindow().ToSharedRef());
 		}
@@ -150,7 +169,7 @@ private:
 		return
 			SNew(STableRow< TSharedRef<FText> >, OwnerTable)
 			[
-				SNew(STextBlock) .Text( Item.Get() )
+				SNew(STextBlock).Text(Item.Get())
 			];
 	}
 
@@ -159,7 +178,7 @@ private:
 		FWidgetPath WidgetPath;
 		TSharedPtr<SWindow> Window = FSlateApplication::Get().FindWidgetWindow(AsShared(), WidgetPath);
 
-		if ( Window.IsValid() )
+		if (Window.IsValid())
 		{
 			Window->RequestDestroyWindow();
 		}
@@ -180,12 +199,12 @@ void FAssetRenameManager::RenameAssets(const TArray<FAssetRenameData>& AssetsAnd
 {
 	// If the asset registry is still loading assets, we cant check for referencers, so we must open the rename dialog
 	FAssetRegistryModule& AssetRegistryModule = FModuleManager::LoadModuleChecked<FAssetRegistryModule>("AssetRegistry");
-	if ( AssetRegistryModule.Get().IsLoadingAssets() )
+	if (AssetRegistryModule.Get().IsLoadingAssets())
 	{
 		// Open a dialog asking the user to wait while assets are being discovered
 		SDiscoveringAssetsDialog::OpenDiscoveringAssetsDialog(
 			SDiscoveringAssetsDialog::FOnAssetsDiscovered::CreateSP(this, &FAssetRenameManager::FixReferencesAndRename, AssetsAndNames)
-			);
+		);
 	}
 	else
 	{
@@ -194,18 +213,36 @@ void FAssetRenameManager::RenameAssets(const TArray<FAssetRenameData>& AssetsAnd
 	}
 }
 
+void FAssetRenameManager::FindSoftReferencesToObject(FSoftObjectPath TargetObject, TArray<UObject*>& ReferencingObjects) const
+{
+	TArray<FAssetRenameDataWithReferencers> AssetsToRename;
+	AssetsToRename.Emplace(FAssetRenameDataWithReferencers(FAssetRenameData(TargetObject, TargetObject, true)));
+
+	// Fill out referencers from asset registry
+	PopulateAssetReferencers(AssetsToRename);
+
+	// Load all referencing objects and find for referencing objects
+	TArray<UPackage*> ReferencingPackagesToSave;
+	LoadReferencingPackages(AssetsToRename, true, ReferencingPackagesToSave, ReferencingObjects);
+}
+
 void FAssetRenameManager::FixReferencesAndRename(TArray<FAssetRenameData> AssetsAndNames) const
 {
+	bool bSoftReferencesOnly = true;
 	// Prep a list of assets to rename with an extra boolean to determine if they should leave a redirector or not
 	TArray<FAssetRenameDataWithReferencers> AssetsToRename;
 	AssetsToRename.Reset(AssetsAndNames.Num());
 	for (const FAssetRenameData& AssetRenameData : AssetsAndNames)
 	{
 		AssetsToRename.Emplace(FAssetRenameDataWithReferencers(AssetRenameData));
+		if (!AssetRenameData.bOnlyFixSoftReferences)
+		{
+			bSoftReferencesOnly = false;
+		}
 	}
 
 	// Warn the user if they are about to rename an asset that is referenced by a CDO
-	auto CDOAssets = FindCDOReferencedAssets(AssetsToRename);
+	TArray<TWeakObjectPtr<UObject>> CDOAssets = FindCDOReferencedAssets(AssetsToRename);
 
 	// Warn the user if there were any references
 	if (CDOAssets.Num())
@@ -220,7 +257,7 @@ void FAssetRenameManager::FixReferencesAndRename(TArray<FAssetRenameData> Assets
 			}
 		}
 
-		const FText MessageText = FText::Format(LOCTEXT("RenameCDOReferences", "The following assets are referenced by one or more Class Default Objects: \n{0}\n\nContinuing with the rename may require code changes to fix these references. Do you wish to continue?"), FText::FromString(AssetNames) );
+		const FText MessageText = FText::Format(LOCTEXT("RenameCDOReferences", "The following assets are referenced by one or more Class Default Objects: \n{0}\n\nContinuing with the rename may require code changes to fix these references. Do you wish to continue?"), FText::FromString(AssetNames));
 		if (FMessageDialog::Open(EAppMsgType::YesNo, MessageText) == EAppReturnType::No)
 		{
 			return;
@@ -231,31 +268,66 @@ void FAssetRenameManager::FixReferencesAndRename(TArray<FAssetRenameData> Assets
 	PopulateAssetReferencers(AssetsToRename);
 
 	// Update the source control state for the packages containing the assets we are renaming if source control is enabled. If source control is enabled and this fails we can not continue.
-	if ( UpdatePackageStatus(AssetsToRename) )
+	if (bSoftReferencesOnly || UpdatePackageStatus(AssetsToRename))
 	{
 		// Detect whether the assets are being referenced by a collection. Assets within a collection must leave a redirector to avoid the collection losing its references.
 		DetectReferencingCollections(AssetsToRename);
 
 		// Load all referencing packages and mark any assets that must have redirectors.
 		TArray<UPackage*> ReferencingPackagesToSave;
-		LoadReferencingPackages(AssetsToRename, ReferencingPackagesToSave);
+		TArray<UObject*> SoftReferencingObjects;
+		LoadReferencingPackages(AssetsToRename, bSoftReferencesOnly, ReferencingPackagesToSave, SoftReferencingObjects);
 
 		// Prompt to check out source package and all referencing packages, leave redirectors for assets referenced by packages that are not checked out and remove those packages from the save list.
 		const bool bUserAcceptedCheckout = CheckOutPackages(AssetsToRename, ReferencingPackagesToSave);
 
-		if ( bUserAcceptedCheckout )
+		if (bUserAcceptedCheckout || bSoftReferencesOnly)
 		{
 			// If any referencing packages are left read-only, the checkout failed or SCC was not enabled. Trim them from the save list and leave redirectors.
 			DetectReadOnlyPackages(AssetsToRename, ReferencingPackagesToSave);
 
-			// Perform the rename, leaving redirectors only for assets which need them
-			PerformAssetRename(AssetsToRename);
+			if (bSoftReferencesOnly)
+			{
+				if (ReferencingPackagesToSave.Num() > 0)
+				{
+					// Only do the rename if there are actually packages with references
+					PerformAssetRename(AssetsToRename);
+
+					for (const FAssetRenameDataWithReferencers& RenameData : AssetsToRename)
+					{
+						// Add source and destination packages so those get saved at the same time
+						UPackage* OldPackage = FindPackage(nullptr, *RenameData.OldObjectPath.GetLongPackageName());
+						UPackage* NewPackage = FindPackage(nullptr, *RenameData.NewObjectPath.GetLongPackageName());
+
+						ReferencingPackagesToSave.AddUnique(OldPackage);
+						ReferencingPackagesToSave.AddUnique(NewPackage);
+					}
 
-			// Save all packages that were referencing any of the assets that were moved without redirectors
-			SaveReferencingPackages(ReferencingPackagesToSave);
+					FString AssetNames;
+					for (UPackage* PackageToSave : ReferencingPackagesToSave)
+					{
+						AssetNames += FString("\n") + PackageToSave->GetName();
+					}
 
-			// Issue post rename event
-			AssetPostRenameEvent.Broadcast(AssetsAndNames);
+					// Warn user before saving referencing packages
+					const FText MessageText = FText::Format(LOCTEXT("SoftReferenceFixedUp", "The following packages were fixed up because they have soft references to a renamed object: \n{0}\n\nDo you want to save them now?\nIf you quit without saving references will be broken!"), FText::FromString(AssetNames));
+					if (FMessageDialog::Open(EAppMsgType::YesNo, MessageText) == EAppReturnType::Yes)
+					{
+						SaveReferencingPackages(ReferencingPackagesToSave);
+					}
+				}
+			}
+			else
+			{
+				// Perform the rename, leaving redirectors only for assets which need them
+				PerformAssetRename(AssetsToRename);
+
+				// Save all packages that were referencing any of the assets that were moved without redirectors
+				SaveReferencingPackages(ReferencingPackagesToSave);
+
+				// Issue post rename event
+				AssetPostRenameEvent.Broadcast(AssetsAndNames);
+			}
 		}
 	}
 
@@ -266,18 +338,21 @@ void FAssetRenameManager::FixReferencesAndRename(TArray<FAssetRenameData> Assets
 TArray<TWeakObjectPtr<UObject>> FAssetRenameManager::FindCDOReferencedAssets(const TArray<FAssetRenameDataWithReferencers>& AssetsToRename) const
 {
 	TArray<TWeakObjectPtr<UObject>> CDOAssets, LocalAssetsToRename;
-	for (auto AssetIt = AssetsToRename.CreateConstIterator(); AssetIt; ++AssetIt)
+	for (const FAssetRenameDataWithReferencers& AssetToRename : AssetsToRename)
 	{
-		LocalAssetsToRename.Push((*AssetIt).Asset);
+		if (AssetToRename.Asset.IsValid())
+		{
+			LocalAssetsToRename.Push(AssetToRename.Asset);
+		}
 	}
 
 	// Run over all CDOs and check for any references to the assets
-	for ( TObjectIterator<UClass> ClassDefaultObjectIt; ClassDefaultObjectIt; ++ClassDefaultObjectIt)
+	for (TObjectIterator<UClass> ClassDefaultObjectIt; ClassDefaultObjectIt; ++ClassDefaultObjectIt)
 	{
 		UClass* Cls = (*ClassDefaultObjectIt);
 		UObject* CDO = Cls->ClassDefaultObject;
 
-		if (!CDO || !CDO->HasAllFlags(RF_ClassDefaultObject) || Cls->ClassGeneratedBy != NULL)
+		if (!CDO || !CDO->HasAllFlags(RF_ClassDefaultObject) || Cls->ClassGeneratedBy != nullptr)
 		{
 			continue;
 		}
@@ -291,9 +366,8 @@ TArray<TWeakObjectPtr<UObject>> FAssetRenameManager::FindCDOReferencedAssets(con
 		for (TFieldIterator<UObjectProperty> PropertyIt(Cls); PropertyIt; ++PropertyIt)
 		{
 			const UObject* Object = PropertyIt->GetPropertyValue(PropertyIt->ContainerPtrToValuePtr<UObject>(CDO));
-			for (auto AssetIt = LocalAssetsToRename.CreateConstIterator(); AssetIt; ++AssetIt)
+			for (const TWeakObjectPtr<UObject> Asset : LocalAssetsToRename)
 			{
-				auto Asset = *AssetIt;
 				if (Object == Asset.Get())
 				{
 					CDOAssets.Push(Asset);
@@ -322,34 +396,47 @@ void FAssetRenameManager::PopulateAssetReferencers(TArray<FAssetRenameDataWithRe
 	TSet<FName> RenamingAssetPackageNames;
 
 	// Get the names of all the packages containing the assets we are renaming so they arent added to the referencing packages list
-	for ( auto AssetIt = AssetsToPopulate.CreateConstIterator(); AssetIt; ++AssetIt )
+	for (FAssetRenameDataWithReferencers& AssetToRename : AssetsToPopulate)
 	{
-		UObject* Asset = (*AssetIt).Asset.Get();
-		if ( Asset )
+		// If we're only fixing soft references we want to check for references inside the original package as we don't save the original package automatically
+		if (!AssetToRename.bOnlyFixSoftReferences)
 		{
-			RenamingAssetPackageNames.Add( FName(*Asset->GetOutermost()->GetName()) );
+			RenamingAssetPackageNames.Add(FName(*AssetToRename.OldObjectPath.GetLongPackageName()));
 		}
 	}
 
 	// Gather all referencing packages for all assets that are being renamed
-	for ( auto AssetIt = AssetsToPopulate.CreateIterator(); AssetIt; ++AssetIt )
+	for (FAssetRenameDataWithReferencers& AssetToRename : AssetsToPopulate)
 	{
-		(*AssetIt).ReferencingPackageNames.Empty();
+		AssetToRename.ReferencingPackageNames.Empty();
 
-		UObject* Asset = (*AssetIt).Asset.Get();
-		if ( Asset )
-		{
-			TArray<FName> Referencers;
-			AssetRegistryModule.Get().GetReferencers(Asset->GetOutermost()->GetFName(), Referencers);
+		FName OldPackageName = FName(*AssetToRename.OldObjectPath.GetLongPackageName());
+
+		TArray<FName> Referencers;
+		AssetRegistryModule.Get().GetReferencers(OldPackageName, Referencers, AssetToRename.bOnlyFixSoftReferences ? EAssetRegistryDependencyType::Soft : EAssetRegistryDependencyType::Packages);
 
-			for ( auto ReferenceIt = Referencers.CreateConstIterator(); ReferenceIt; ++ReferenceIt )
+		for (const FName& ReferencingPackageName : Referencers)
+		{
+			if (!RenamingAssetPackageNames.Contains(ReferencingPackageName))
 			{
-				const FName& ReferencingPackageName = *ReferenceIt;
+				AssetToRename.ReferencingPackageNames.AddUnique(ReferencingPackageName);
+			}
+		}
 
-				if ( !RenamingAssetPackageNames.Contains(ReferencingPackageName) )
-				{
-					(*AssetIt).ReferencingPackageNames.AddUnique( ReferencingPackageName );
-				}
+		if (AssetToRename.bOnlyFixSoftReferences)
+		{
+			AssetToRename.ReferencingPackageNames.AddUnique(FName(*AssetToRename.OldObjectPath.GetLongPackageName()));
+			AssetToRename.ReferencingPackageNames.AddUnique(FName(*AssetToRename.NewObjectPath.GetLongPackageName()));
+
+			// Add dirty packages and the package that owns the reference. They will get filtered out in LoadReferencingPackages if they aren't valid
+			TArray<UPackage*> ExtraPackagesToCheckForSoftReferences;
+
+			FEditorFileUtils::GetDirtyWorldPackages(ExtraPackagesToCheckForSoftReferences);
+			FEditorFileUtils::GetDirtyContentPackages(ExtraPackagesToCheckForSoftReferences);
+
+			for (UPackage* Package : ExtraPackagesToCheckForSoftReferences)
+			{
+				AssetToRename.ReferencingPackageNames.AddUnique(Package->GetFName());
 			}
 		}
 	}
@@ -357,24 +444,24 @@ void FAssetRenameManager::PopulateAssetReferencers(TArray<FAssetRenameDataWithRe
 
 bool FAssetRenameManager::UpdatePackageStatus(const TArray<FAssetRenameDataWithReferencers>& AssetsToRename) const
 {
-	if ( ISourceControlModule::Get().IsEnabled() )
+	if (ISourceControlModule::Get().IsEnabled())
 	{
 		ISourceControlProvider& SourceControlProvider = ISourceControlModule::Get().GetProvider();
 
 		// Update the source control server availability to make sure we can do the rename operation
 		SourceControlProvider.Login();
-		if ( !SourceControlProvider.IsAvailable() )
+		if (!SourceControlProvider.IsAvailable())
 		{
-			FMessageDialog::Open( EAppMsgType::Ok, NSLOCTEXT("UnrealEd", "SourceControl_ServerUnresponsive", "Source Control is unresponsive. Please check your connection and try again.") );
+			FMessageDialog::Open(EAppMsgType::Ok, NSLOCTEXT("UnrealEd", "SourceControl_ServerUnresponsive", "Source Control is unresponsive. Please check your connection and try again."));
 			return false;
 		}
 
 		// Gather asset package names to update SCC states in a single SCC request
 		TArray<UPackage*> PackagesToUpdate;
-		for ( auto AssetIt = AssetsToRename.CreateConstIterator(); AssetIt; ++AssetIt )
+		for (auto AssetIt = AssetsToRename.CreateConstIterator(); AssetIt; ++AssetIt)
 		{
 			UObject* Asset = (*AssetIt).Asset.Get();
-			if ( Asset )
+			if (Asset)
 			{
 				PackagesToUpdate.AddUnique(Asset->GetOutermost());
 			}
@@ -386,26 +473,26 @@ bool FAssetRenameManager::UpdatePackageStatus(const TArray<FAssetRenameDataWithR
 	return true;
 }
 
-void FAssetRenameManager::LoadReferencingPackages(TArray<FAssetRenameDataWithReferencers>& AssetsToRename, TArray<UPackage*>& OutReferencingPackagesToSave) const
+void FAssetRenameManager::LoadReferencingPackages(TArray<FAssetRenameDataWithReferencers>& AssetsToRename, bool bLoadAllPackages, TArray<UPackage*>& OutReferencingPackagesToSave, TArray<UObject*>& OutSoftReferencingObjects) const
 {
 	const FText ReferenceUpdateSlowTask = LOCTEXT("ReferenceUpdateSlowTask", "Updating Asset References");
-	GWarn->BeginSlowTask( ReferenceUpdateSlowTask, true );
+	GWarn->BeginSlowTask(ReferenceUpdateSlowTask, true);
 
 	ISourceControlProvider& SourceControlProvider = ISourceControlModule::Get().GetProvider();
 
-	for ( int32 AssetIdx = 0; AssetIdx < AssetsToRename.Num(); ++AssetIdx )
+	for (int32 AssetIdx = 0; AssetIdx < AssetsToRename.Num(); ++AssetIdx)
 	{
-		GWarn->StatusUpdate( AssetIdx, AssetsToRename.Num(), ReferenceUpdateSlowTask );
+		GWarn->StatusUpdate(AssetIdx, AssetsToRename.Num(), ReferenceUpdateSlowTask);
 
 		FAssetRenameDataWithReferencers& RenameData = AssetsToRename[AssetIdx];
 
 		UObject* Asset = RenameData.Asset.Get();
-		if ( Asset )
+		if (Asset)
 		{
 			// Make sure this asset is local. Only local assets should be renamed without a redirector
 			FSourceControlStatePtr SourceControlState = SourceControlProvider.GetState(Asset->GetOutermost(), EStateCacheUsage::ForceUpdate);
 			const bool bLocalFile = !SourceControlState.IsValid() || SourceControlState->IsAdded() || !SourceControlState->IsSourceControlled() || SourceControlState->IsIgnored();
-			if ( !bLocalFile )
+			if (!bLocalFile)
 			{
 				// If this asset is locked or not current, mark it failed to prevent it from being renamed
 				if (SourceControlState->IsCheckedOutOther())
@@ -421,46 +508,75 @@ void FAssetRenameManager::LoadReferencingPackages(TArray<FAssetRenameDataWithRef
 
 				// This asset is not local. It is not safe to rename it without leaving a redirector
 				RenameData.bCreateRedirector = true;
-				continue;
+				if (!bLoadAllPackages)
+				{
+					continue;
+				}
 			}
 		}
 		else
 		{
-			// The asset for this rename must have been GCed or is otherwise invalid. Skip it
-			continue;
+			// The asset for this rename must have been GCed or is otherwise invalid. Skip it unless this is a soft reference only fix
+			if (!bLoadAllPackages)
+			{
+				continue;
+			}
 		}
 
+		TMap<FSoftObjectPath, FSoftObjectPath> ModifiedPaths;
+		ModifiedPaths.Add(RenameData.OldObjectPath, RenameData.NewObjectPath);
+
 		TArray<UPackage*> PackagesToSaveForThisAsset;
 		bool bAllPackagesLoadedForThisAsset = true;
-		for ( auto PackageNameIt = RenameData.ReferencingPackageNames.CreateConstIterator(); PackageNameIt; ++PackageNameIt )
+		for (int32 i = 0; i < RenameData.ReferencingPackageNames.Num(); i++)
 		{
-			const FString PackageName = (*PackageNameIt).ToString();
-
+			FName PackageName = RenameData.ReferencingPackageNames[i];
 			// Check if the package is a map before loading it!
-			if ( FEditorFileUtils::IsMapPackageAsset(PackageName) )
+			if (!bLoadAllPackages && FEditorFileUtils::IsMapPackageAsset(PackageName.ToString()))
 			{
 				// This reference was a map package, don't load it and leave a redirector for this asset
+				// For subobjects we want to load maps packages and treat them normally
 				RenameData.bCreateRedirector = true;
 				bAllPackagesLoadedForThisAsset = false;
 				break;
 			}
 
-			UPackage* Package = LoadPackage(NULL, *PackageName, LOAD_None);
-			if ( Package )
-			{
-				PackagesToSaveForThisAsset.Add(Package);
+			UPackage* Package = LoadPackage(nullptr, *PackageName.ToString(), LOAD_None);
+			if (Package)
+			{				
+				bool bFoundSoftReference = CheckPackageForSoftObjectReferences(Package, ModifiedPaths, OutSoftReferencingObjects);
+
+				// Only add to list if we're doing a hard reference fixup or we found a soft reference
+				bool bAdd = !RenameData.bOnlyFixSoftReferences || bFoundSoftReference;
+
+				if (bAdd)
+				{
+					PackagesToSaveForThisAsset.Add(Package);
+				}
+				else
+				{
+					// This package does not actually reference the asset, so remove it
+					RenameData.ReferencingPackageNames.RemoveAt(i);
+					i--;
+				}
 			}
 			else
 			{
 				RenameData.bCreateRedirector = true;
-				bAllPackagesLoadedForThisAsset = false;
-				break;
+				if (!bLoadAllPackages)
+				{
+					bAllPackagesLoadedForThisAsset = false;
+					break;
+				}
 			}
 		}
 
-		if ( bAllPackagesLoadedForThisAsset )
+		if (bAllPackagesLoadedForThisAsset)
 		{
-			OutReferencingPackagesToSave.Append(PackagesToSaveForThisAsset);
+			for (UPackage* Package : PackagesToSaveForThisAsset)
+			{
+				OutReferencingPackagesToSave.AddUnique(Package);
+			}
 		}
 	}
 
@@ -470,7 +586,7 @@ void FAssetRenameManager::LoadReferencingPackages(TArray<FAssetRenameDataWithRef
 bool FAssetRenameManager::CheckOutPackages(TArray<FAssetRenameDataWithReferencers>& AssetsToRename, TArray<UPackage*>& InOutReferencingPackagesToSave) const
 {
 	bool bUserAcceptedCheckout = true;
-	
+
 	// Build list of packages to check out: the source package and any referencing packages (in the case that we do not create a redirector)
 	TArray<UPackage*> PackagesToCheckOut;
 	PackagesToCheckOut.Reset(AssetsToRename.Num() + InOutReferencingPackagesToSave.Num());
@@ -493,20 +609,20 @@ bool FAssetRenameManager::CheckOutPackages(TArray<FAssetRenameDataWithReferencer
 	{
 		TArray<UPackage*> PackagesCheckedOutOrMadeWritable;
 		TArray<UPackage*> PackagesNotNeedingCheckout;
-		bUserAcceptedCheckout = FEditorFileUtils::PromptToCheckoutPackages( false, PackagesToCheckOut, &PackagesCheckedOutOrMadeWritable, &PackagesNotNeedingCheckout );
-		if ( bUserAcceptedCheckout )
+		bUserAcceptedCheckout = FEditorFileUtils::PromptToCheckoutPackages(false, PackagesToCheckOut, &PackagesCheckedOutOrMadeWritable, &PackagesNotNeedingCheckout);
+		if (bUserAcceptedCheckout)
 		{
 			// Make a list of any packages in the list which weren't checked out for some reason
 			TArray<UPackage*> PackagesThatCouldNotBeCheckedOut = PackagesToCheckOut;
 
-			for ( auto PackageIt = PackagesCheckedOutOrMadeWritable.CreateConstIterator(); PackageIt; ++PackageIt )
+			for (UPackage* Package : PackagesCheckedOutOrMadeWritable)
 			{
-				PackagesThatCouldNotBeCheckedOut.Remove(*PackageIt);
+				PackagesThatCouldNotBeCheckedOut.Remove(Package);
 			}
 
-			for ( auto PackageIt = PackagesNotNeedingCheckout.CreateConstIterator(); PackageIt; ++PackageIt )
+			for (UPackage* Package : PackagesNotNeedingCheckout)
 			{
-				PackagesThatCouldNotBeCheckedOut.Remove(*PackageIt);
+				PackagesThatCouldNotBeCheckedOut.Remove(Package);
 			}
 
 			// If there's anything which couldn't be checked out, abort the operation.
@@ -524,7 +640,7 @@ void FAssetRenameManager::DetectReferencingCollections(TArray<FAssetRenameDataWi
 {
 	FCollectionManagerModule& CollectionManagerModule = FCollectionManagerModule::GetModule();
 
-	for (auto& AssetToRename : AssetsToRename)
+	for (FAssetRenameDataWithReferencers& AssetToRename : AssetsToRename)
 	{
 		if (AssetToRename.Asset.IsValid())
 		{
@@ -542,26 +658,25 @@ void FAssetRenameManager::DetectReferencingCollections(TArray<FAssetRenameDataWi
 void FAssetRenameManager::DetectReadOnlyPackages(TArray<FAssetRenameDataWithReferencers>& AssetsToRename, TArray<UPackage*>& InOutReferencingPackagesToSave) const
 {
 	// For each valid package...
-	for ( int32 PackageIdx = InOutReferencingPackagesToSave.Num() - 1; PackageIdx >= 0; --PackageIdx )
+	for (int32 PackageIdx = InOutReferencingPackagesToSave.Num() - 1; PackageIdx >= 0; --PackageIdx)
 	{
 		UPackage* Package = InOutReferencingPackagesToSave[PackageIdx];
 
-		if ( Package )
+		if (Package)
 		{
 			// Find the package filename
 			FString Filename;
-			if ( FPackageName::DoesPackageExist(Package->GetName(), NULL, &Filename) )
+			if (FPackageName::DoesPackageExist(Package->GetName(), nullptr, &Filename))
 			{
 				// If the file is read only
-				if ( IFileManager::Get().IsReadOnly(*Filename) )
+				if (IFileManager::Get().IsReadOnly(*Filename))
 				{
 					FName PackageName = Package->GetFName();
 
 					// Find all assets that were referenced by this package to create a redirector when named
-					for ( auto RenameDataIt = AssetsToRename.CreateIterator(); RenameDataIt; ++RenameDataIt )
+					for (FAssetRenameDataWithReferencers& RenameData : AssetsToRename)
 					{
-						FAssetRenameDataWithReferencers& RenameData = *RenameDataIt;
-						if ( RenameData.ReferencingPackageNames.Contains(PackageName) )
+						if (RenameData.ReferencingPackageNames.Contains(PackageName))
 						{
 							RenameData.bCreateRedirector = true;
 						}
@@ -574,33 +689,93 @@ void FAssetRenameManager::DetectReadOnlyPackages(TArray<FAssetRenameDataWithRefe
 		}
 	}
 }
-
-void FAssetRenameManager::RenameReferencingStringAssetReferences(const TArray<UPackage *> PackagesToCheck, const TMap<FString, FString>& AssetRedirectorMap)
+struct FSoftObjectPathRenameSerializer : public FArchiveUObject
 {
-	struct FStringAssetReferenceRenameSerializer : public FArchiveUObject
+	void StartSerializingObject(UObject* InCurrentObject)
+	{ 
+		CurrentObject = InCurrentObject;
+		bFoundReference = false; 
+	}
+	bool HasFoundReference() const 
+	{ 
+		return bFoundReference; 
+	}
+
+	FSoftObjectPathRenameSerializer(const TMap<FSoftObjectPath, FSoftObjectPath>& InRedirectorMap, bool bInCheckOnly)
+		: RedirectorMap(InRedirectorMap)
+		, CurrentObject(nullptr)
+		, bSearchOnly(bInCheckOnly)
+		, bFoundReference(false)
 	{
-		FStringAssetReferenceRenameSerializer()
+		// Mark it as saving to correctly process all references
+		ArIsSaving = true;
+	}
+
+	FArchive& operator<<(FSoftObjectPath& Value)
+	{
+		FString SubPath = Value.GetSubPathString();
+		for (const TPair<FSoftObjectPath, FSoftObjectPath>& Pair : RedirectorMap)
 		{
-			// Mark it as saving to correctly process all references
-			ArIsSaving = true;
+			if (Pair.Key.GetAssetPathName() == Value.GetAssetPathName())
+			{
+				// Same asset, fix sub path. Asset will be fixed by normal serializePath call below
+				FString CheckSubPath = Pair.Key.GetSubPathString();
+
+				if (CheckSubPath.IsEmpty() || SubPath == CheckSubPath || SubPath.StartsWith(CheckSubPath + TEXT(".")))
+				{
+					bFoundReference = true;
+
+					if (!bSearchOnly)
+					{
+						if (CurrentObject)
+						{
+							CurrentObject->Modify(true);
+						}
+						SubPath.ReplaceInline(*CheckSubPath, *Pair.Value.GetSubPathString());
+						Value = FSoftObjectPath(Pair.Value.GetAssetPathName(), SubPath);
+					}
+					break;
+				}
+			}
 		}
-	};
 
+		return *this;
+	}
+
+private:
+	const TMap<FSoftObjectPath, FSoftObjectPath>& RedirectorMap;
+	UObject* CurrentObject;
+	bool bSearchOnly;
+	bool bFoundReference;
+
+};
+
+void FAssetRenameManager::RenameReferencingSoftObjectPaths(const TArray<UPackage *> PackagesToCheck, const TMap<FSoftObjectPath, FSoftObjectPath>& AssetRedirectorMap)
+{
 	// Add redirects as needed
-	for (const TPair<FString, FString>& Pair : AssetRedirectorMap)
+	for (const TPair<FSoftObjectPath, FSoftObjectPath>& Pair : AssetRedirectorMap)
 	{
-		GRedirectCollector.AddAssetPathRedirection(Pair.Key, Pair.Value);
+		if (Pair.Key.IsAsset())
+		{
+			GRedirectCollector.AddAssetPathRedirection(Pair.Key.GetAssetPathName(), Pair.Value.GetAssetPathName());
+		}
 	}
 
-	FStringAssetReferenceRenameSerializer RenameSerializer;
+	FSoftObjectPathRenameSerializer RenameSerializer(AssetRedirectorMap, false);
 
-	for (auto* Package : PackagesToCheck)
+	for (UPackage* Package : PackagesToCheck)
 	{
 		TArray<UObject*> ObjectsInPackage;
 		GetObjectsWithOuter(Package, ObjectsInPackage);
 
-		for (auto* Object : ObjectsInPackage)
+		for (UObject* Object : ObjectsInPackage)
 		{
+			if (Object->IsPendingKill())
+			{
+				continue;
+			}
+
+			RenameSerializer.StartSerializingObject(Object);
 			Object->Serialize(RenameSerializer);
 
 			if (UBlueprint* Blueprint = Cast<UBlueprint>(Object))
@@ -610,12 +785,42 @@ void FAssetRenameManager::RenameReferencingStringAssetReferences(const TArray<UP
 			}
 		}
 	}
+
+	// Invalidate the soft object tag as we have created new valid paths
+	FSoftObjectPath::InvalidateTag();
+}
+
+bool FAssetRenameManager::CheckPackageForSoftObjectReferences(UPackage* Package, const TMap<FSoftObjectPath, FSoftObjectPath>& AssetRedirectorMap, TArray<UObject*>& OutReferencingObjects)
+{	
+	bool bFoundReference = false;
+	FSoftObjectPathRenameSerializer CheckSerializer(AssetRedirectorMap, true);
+
+	TArray<UObject*> ObjectsInPackage;
+	GetObjectsWithOuter(Package, ObjectsInPackage);
+
+	for (UObject* Object : ObjectsInPackage)
+	{
+		if (Object->IsPendingKill())
+		{
+			continue;
+		}
+
+		CheckSerializer.StartSerializingObject(Object);
+		Object->Serialize(CheckSerializer);
+
+		if (CheckSerializer.HasFoundReference())
+		{
+			bFoundReference = true;
+			OutReferencingObjects.AddUnique(Object);
+		}
+	}
+	return bFoundReference;
 }
 
 void FAssetRenameManager::PerformAssetRename(TArray<FAssetRenameDataWithReferencers>& AssetsToRename) const
 {
 	const FText AssetRenameSlowTask = LOCTEXT("AssetRenameSlowTask", "Renaming Assets");
-	GWarn->BeginSlowTask( AssetRenameSlowTask, true );
+	GWarn->BeginSlowTask(AssetRenameSlowTask, true);
 
 	/**
 	 * We need to collect and check those cause dependency graph is only
@@ -623,108 +828,111 @@ void FAssetRenameManager::PerformAssetRename(TArray<FAssetRenameDataWithReferenc
 	 * objects. It is only needed for string references as in memory references
 	 * for other objects are pointers, so renames doesn't apply to those.
 	 */
-	TArray<UPackage *> DirtyPackagesToCheckForStringReferences;
-	
-	FEditorFileUtils::GetDirtyWorldPackages(DirtyPackagesToCheckForStringReferences);
-	FEditorFileUtils::GetDirtyContentPackages(DirtyPackagesToCheckForStringReferences);
+	TArray<UPackage *> DirtyPackagesToCheckForSoftReferences;
+
+	FEditorFileUtils::GetDirtyWorldPackages(DirtyPackagesToCheckForSoftReferences);
+	FEditorFileUtils::GetDirtyContentPackages(DirtyPackagesToCheckForSoftReferences);
 
 	TArray<UPackage*> PackagesToSave;
 	TArray<UPackage*> PotentialPackagesToDelete;
-	for ( int32 AssetIdx = 0; AssetIdx < AssetsToRename.Num(); ++AssetIdx )
+	for (int32 AssetIdx = 0; AssetIdx < AssetsToRename.Num(); ++AssetIdx)
 	{
-		GWarn->StatusUpdate( AssetIdx, AssetsToRename.Num(), AssetRenameSlowTask );
+		GWarn->StatusUpdate(AssetIdx, AssetsToRename.Num(), AssetRenameSlowTask);
 
 		FAssetRenameDataWithReferencers& RenameData = AssetsToRename[AssetIdx];
 
-		if ( RenameData.bRenameFailed )
+		if (RenameData.bRenameFailed)
 		{
 			// The rename failed at some earlier step, skip this asset
 			continue;
 		}
 
 		UObject* Asset = RenameData.Asset.Get();
+		TArray<UPackage *> PackagesToCheckForSoftReferences;
 
-		if ( !Asset )
+		if (!RenameData.bOnlyFixSoftReferences)
 		{
-			// This asset was invalid or GCed before the rename could occur
-			RenameData.bRenameFailed = true;
-			continue;
-		}
-
-		FString OldAssetPath = Asset->GetPathName();
+			// If bOnlyFixSoftReferences was set these got appended in find references
+			PackagesToCheckForSoftReferences.Append(DirtyPackagesToCheckForSoftReferences);
 
-		ObjectTools::FPackageGroupName PGN;
-		PGN.ObjectName = RenameData.NewName;
-		PGN.GroupName = TEXT("");
-		PGN.PackageName = RenameData.PackagePath / PGN.ObjectName;
-		const bool bLeaveRedirector = RenameData.bCreateRedirector;
+			if (!Asset)
+			{
+				// This asset was invalid or GCed before the rename could occur
+				RenameData.bRenameFailed = true;
+				continue;
+			}
 
-		UPackage* OldPackage = Asset->GetOutermost();
-		bool bOldPackageAddedToRootSet = false;
-		if ( !bLeaveRedirector && !OldPackage->IsRooted() )
-		{
-			bOldPackageAddedToRootSet = true;
-			OldPackage->AddToRoot();
-		}
+			ObjectTools::FPackageGroupName PGN;
+			PGN.ObjectName = RenameData.NewName;
+			PGN.GroupName = TEXT("");
+			PGN.PackageName = RenameData.NewPackagePath / PGN.ObjectName;
+			const bool bLeaveRedirector = RenameData.bCreateRedirector;
 
-		TSet<UPackage*> ObjectsUserRefusedToFullyLoad;
-		FText ErrorMessage;
-		if ( ObjectTools::RenameSingleObject(Asset, PGN, ObjectsUserRefusedToFullyLoad, ErrorMessage, NULL, bLeaveRedirector) )
-		{
-			PackagesToSave.AddUnique( Asset->GetOutermost() );
-
-			// Automatically save renamed assets
-			if ( bLeaveRedirector )
+			UPackage* OldPackage = Asset->GetOutermost();
+			bool bOldPackageAddedToRootSet = false;
+			if (!bLeaveRedirector && !OldPackage->IsRooted())
 			{
-				PackagesToSave.AddUnique( OldPackage );
+				bOldPackageAddedToRootSet = true;
+				OldPackage->AddToRoot();
 			}
-			else if ( bOldPackageAddedToRootSet )
+
+			TSet<UPackage*> ObjectsUserRefusedToFullyLoad;
+			FText ErrorMessage;
+			if (ObjectTools::RenameSingleObject(Asset, PGN, ObjectsUserRefusedToFullyLoad, ErrorMessage, nullptr, bLeaveRedirector))
 			{
-				// Since we did not leave a redirector and the old package wasnt already rooted, attempt to delete it when we are done. 
-				PotentialPackagesToDelete.AddUnique(OldPackage);
+				PackagesToSave.AddUnique(Asset->GetOutermost());
+
+				// Automatically save renamed assets
+				if (bLeaveRedirector)
+				{
+					PackagesToSave.AddUnique(OldPackage);
+				}
+				else if (bOldPackageAddedToRootSet)
+				{
+					// Since we did not leave a redirector and the old package wasnt already rooted, attempt to delete it when we are done. 
+					PotentialPackagesToDelete.AddUnique(OldPackage);
+				}
 			}
-		}
-		else
-		{
-			// No need to keep the old package rooted, the asset was never renamed out of it
-			if ( bOldPackageAddedToRootSet )
+			else
 			{
-				OldPackage->RemoveFromRoot();
-			}
+				// No need to keep the old package rooted, the asset was never renamed out of it
+				if (bOldPackageAddedToRootSet)
+				{
+					OldPackage->RemoveFromRoot();
+				}
 
-			// Mark the rename as a failure to report it later
-			RenameData.bRenameFailed = true;
-			RenameData.FailureReason = ErrorMessage;
+				// Mark the rename as a failure to report it later
+				RenameData.bRenameFailed = true;
+				RenameData.FailureReason = ErrorMessage;
+			}
 		}
 
-		TArray<UPackage *> PackagesToCheck(DirtyPackagesToCheckForStringReferences);
-
-		for (auto PackageNameIt = RenameData.ReferencingPackageNames.CreateConstIterator(); PackageNameIt; ++PackageNameIt)
+		for (FName PackageName : RenameData.ReferencingPackageNames)
 		{
-			UPackage* PackageToCheck = FindPackage(NULL, *PackageNameIt->ToString());
+			UPackage* PackageToCheck = FindPackage(nullptr, *PackageName.ToString());
 			if (PackageToCheck)
 			{
-				PackagesToCheck.Add(PackageToCheck);
+				PackagesToCheckForSoftReferences.Add(PackageToCheck);
 			}
 		}
 
-		TMap<FString, FString> RedirectorMap;
-		RedirectorMap.Add(OldAssetPath, Asset->GetPathName());
+		TMap<FSoftObjectPath, FSoftObjectPath> RedirectorMap;
+		RedirectorMap.Add(RenameData.OldObjectPath, RenameData.NewObjectPath);
 
 		if (UBlueprint* Blueprint = Cast<UBlueprint>(Asset))
 		{
 			// Add redirect for class and default as well
-			RedirectorMap.Add(FString::Printf(TEXT("%s_C"), *OldAssetPath), FString::Printf(TEXT("%s_C"), *Asset->GetPathName()));
-			RedirectorMap.Add(FString::Printf(TEXT("Default__%s_C"), *OldAssetPath), FString::Printf(TEXT("Default__%s_C"), *Asset->GetPathName()));
+			RedirectorMap.Add(FString::Printf(TEXT("%s_C"), *RenameData.OldObjectPath.ToString()), FString::Printf(TEXT("%s_C"), *RenameData.NewObjectPath.ToString()));
+			RedirectorMap.Add(FString::Printf(TEXT("%s.Default__%s_C"), *RenameData.OldObjectPath.GetLongPackageName(), *RenameData.OldObjectPath.GetAssetName()), FString::Printf(TEXT("%s.Default__%s_C"), *RenameData.NewObjectPath.GetLongPackageName(), *RenameData.NewObjectPath.GetAssetName()));
 		}
 
-		RenameReferencingStringAssetReferences(PackagesToCheck, RedirectorMap);
+		RenameReferencingSoftObjectPaths(PackagesToCheckForSoftReferences, RedirectorMap);
 	}
 
 	GWarn->EndSlowTask();
 
 	// Save all renamed assets and any redirectors that were left behind
-	if ( PackagesToSave.Num() > 0 )
+	if (PackagesToSave.Num() > 0)
 	{
 		const bool bCheckDirty = false;
 		const bool bPromptToSave = false;
@@ -735,22 +943,25 @@ void FAssetRenameManager::PerformAssetRename(TArray<FAssetRenameDataWithReferenc
 	}
 
 	// Now branch the files in source control if possible
-	for (const auto& AssetToRename : AssetsToRename)
+	for (const FAssetRenameDataWithReferencers& RenameData : AssetsToRename)
 	{
+		UPackage* OldPackage = FindPackage(nullptr, *RenameData.OldObjectPath.GetLongPackageName());
+		UPackage* NewPackage = FindPackage(nullptr, *RenameData.NewObjectPath.GetLongPackageName());
+
 		// If something went wrong when saving and the new asset does not exist on disk, don't branch it
 		// as it will just create a copy and any attempt to load it will result in crashes.
-		if (FPackageName::DoesPackageExist(AssetToRename.Asset->GetOutermost()->GetFName().ToString()))
+		if (!RenameData.bOnlyFixSoftReferences && FPackageName::DoesPackageExist(NewPackage->GetName()))
 		{
-			SourceControlHelpers::BranchPackage(AssetToRename.Asset->GetOutermost(), FindPackage(nullptr, *AssetToRename.OriginalAssetPath));
+			SourceControlHelpers::BranchPackage(NewPackage, OldPackage);
 		}
 	}
 
 	// Clean up all packages that were left empty
-	if ( PotentialPackagesToDelete.Num() > 0 )
+	if (PotentialPackagesToDelete.Num() > 0)
 	{
-		for ( auto PackageIt = PotentialPackagesToDelete.CreateConstIterator(); PackageIt; ++PackageIt )
+		for (UPackage* Package : PotentialPackagesToDelete)
 		{
-			(*PackageIt)->RemoveFromRoot();
+			Package->RemoveFromRoot();
 		}
 
 		ObjectTools::CleanupAfterSuccessfulDelete(PotentialPackagesToDelete);
@@ -759,7 +970,7 @@ void FAssetRenameManager::PerformAssetRename(TArray<FAssetRenameDataWithReferenc
 
 void FAssetRenameManager::SaveReferencingPackages(const TArray<UPackage*>& ReferencingPackagesToSave) const
 {
-	if ( ReferencingPackagesToSave.Num() > 0 )
+	if (ReferencingPackagesToSave.Num() > 0)
 	{
 		const bool bCheckDirty = false;
 		const bool bPromptToSave = false;
@@ -772,13 +983,12 @@ void FAssetRenameManager::SaveReferencingPackages(const TArray<UPackage*>& Refer
 void FAssetRenameManager::ReportFailures(const TArray<FAssetRenameDataWithReferencers>& AssetsToRename) const
 {
 	TArray<FText> FailedRenames;
-	for ( auto RenameIt = AssetsToRename.CreateConstIterator(); RenameIt; ++RenameIt )
+	for (const FAssetRenameDataWithReferencers& RenameData : AssetsToRename)
 	{
-		const FAssetRenameDataWithReferencers& RenameData = *RenameIt;
-		if ( RenameData.bRenameFailed )
+		if (RenameData.bRenameFailed)
 		{
 			UObject* Asset = RenameData.Asset.Get();
-			if ( Asset )
+			if (Asset)
 			{
 				FFormatNamedArguments Args;
 				Args.Add(TEXT("FailureReason"), RenameData.FailureReason);
@@ -793,7 +1003,7 @@ void FAssetRenameManager::ReportFailures(const TArray<FAssetRenameDataWithRefere
 		}
 	}
 
-	if ( FailedRenames.Num() > 0 )
+	if (FailedRenames.Num() > 0)
 	{
 		SRenameFailures::OpenRenameFailuresDialog(FailedRenames);
 	}
diff --git a/Engine/Source/Developer/AssetTools/Private/AssetRenameManager.h b/Engine/Source/Developer/AssetTools/Private/AssetRenameManager.h
--- a/Engine/Source/Developer/AssetTools/Private/AssetRenameManager.h
+++ b/Engine/Source/Developer/AssetTools/Private/AssetRenameManager.h
@@ -4,6 +4,7 @@
 
 #include "CoreMinimal.h"
 #include "IAssetTools.h"
+#include "UObject/SoftObjectPath.h"
 
 struct FAssetRenameDataWithReferencers;
 
@@ -22,16 +23,22 @@ public:
 	/** Renames assets using the specified names. */
 	void RenameAssets(const TArray<FAssetRenameData>& AssetsAndNames) const;
 
+	/** Returns list of objects that soft reference the given soft object path. This will load assets into memory to verify */
+	void FindSoftReferencesToObject(FSoftObjectPath TargetObject, TArray<UObject*>& ReferencingObjects) const;
+
 	/** Accessor for post rename event */
 	FAssetPostRenameEvent& OnAssetPostRenameEvent() { return AssetPostRenameEvent; }
 
 	/**
-	 * Function that renames all FStringAssetReference object with the old asset path to the new one.
+	 * Function that renames all FSoftObjectPath object with the old asset path to the new one.
 	 *
-	 * @param PackagesToCheck Packages to check for referencing FStringAssetReference.
+	 * @param PackagesToCheck Packages to check for referencing FSoftObjectPath.
 	 * @param AssetRedirectorMap Map from old asset path to new asset path
 	 */
-	static void RenameReferencingStringAssetReferences(const TArray<UPackage *> PackagesToCheck, const TMap<FString, FString>& AssetRedirectorMap);
+	static void RenameReferencingSoftObjectPaths(TArray<UPackage*> PackagesToCheck, const TMap<FSoftObjectPath, FSoftObjectPath>& AssetRedirectorMap);
+
+	/** Filters packages list depending on if it actually has soft object paths pointing to the specific object being renamed */
+	static bool CheckPackageForSoftObjectReferences(UPackage* Package, const TMap<FSoftObjectPath, FSoftObjectPath>& AssetRedirectorMap, TArray<UObject*>& OutReferencingObjects);
 
 private:
 	/** Attempts to load and fix redirector references for the supplied assets */
@@ -49,8 +56,9 @@ private:
 	/** 
 	  * Loads all referencing packages to assets in AssetsToRename, finds assets whose references can
 	  * not be fixed up to mark that a redirector should be left, and returns a list of referencing packages to save.
+	  * if bFindAllSoftObjectReferences is true, it will load all referencing packages even if they can't be checked out
 	  */
-	void LoadReferencingPackages(TArray<FAssetRenameDataWithReferencers>& AssetsToRename, TArray<UPackage*>& OutReferencingPackagesToSave) const;
+	void LoadReferencingPackages(TArray<FAssetRenameDataWithReferencers>& AssetsToRename, bool bLoadAllPackages, TArray<UPackage*>& OutReferencingPackagesToSave, TArray<UObject*>& OutSoftReferencingObjects) const;
 
 	/** 
 	  * Prompts to check out the source package and all referencing packages and marks assets whose referencing packages were not checked out to leave a redirector.
diff --git a/Engine/Source/Developer/AssetTools/Private/AssetTools.cpp b/Engine/Source/Developer/AssetTools/Private/AssetTools.cpp
--- a/Engine/Source/Developer/AssetTools/Private/AssetTools.cpp
+++ b/Engine/Source/Developer/AssetTools/Private/AssetTools.cpp
@@ -586,6 +586,11 @@ void UAssetToolsImpl::RenameAssets(const TArray<FAssetRenameData>& AssetsAndName
 	AssetRenameManager->RenameAssets(AssetsAndNames);
 }
 
+void UAssetToolsImpl::FindSoftReferencesToObject(FSoftObjectPath TargetObject, TArray<UObject*>& ReferencingObjects) const
+{
+	AssetRenameManager->FindSoftReferencesToObject(TargetObject, ReferencingObjects);
+}
+
 TArray<UObject*> UAssetToolsImpl::ImportAssets(const FString& DestinationPath)
 {
 	return ImportAssetsWithDialog(DestinationPath);
diff --git a/Engine/Source/Developer/AssetTools/Private/AssetTools.h b/Engine/Source/Developer/AssetTools/Private/AssetTools.h
--- a/Engine/Source/Developer/AssetTools/Private/AssetTools.h
+++ b/Engine/Source/Developer/AssetTools/Private/AssetTools.h
@@ -71,6 +71,7 @@ public:
 	virtual UObject* DuplicateAsset(const FString& AssetName, const FString& PackagePath, UObject* OriginalObject) override;
 	virtual UObject* DuplicateAssetWithDialog(const FString& AssetName, const FString& PackagePath, UObject* OriginalObject) override;
 	virtual void RenameAssets(const TArray<FAssetRenameData>& AssetsAndNames) const override;
+	virtual void FindSoftReferencesToObject(FSoftObjectPath TargetObject, TArray<UObject*>& ReferencingObjects) const override;
 	virtual TArray<UObject*> ImportAssets(const FString& DestinationPath) override;
 	virtual TArray<UObject*> ImportAssetsWithDialog(const FString& DestinationPath) override;
 	virtual TArray<UObject*> ImportAssets(const TArray<FString>& Files, const FString& DestinationPath, UFactory* ChosenFactory, bool bSyncToBrowser = true, TArray<TPair<FString, FString>>* FilesAndDestinations = nullptr) const override;
diff --git a/Engine/Source/Developer/AssetTools/Private/AssetTypeActions/AssetTypeActions_Enum.cpp b/Engine/Source/Developer/AssetTools/Private/AssetTypeActions/AssetTypeActions_Enum.cpp
--- a/Engine/Source/Developer/AssetTools/Private/AssetTypeActions/AssetTypeActions_Enum.cpp
+++ b/Engine/Source/Developer/AssetTools/Private/AssetTypeActions/AssetTypeActions_Enum.cpp
@@ -2,9 +2,15 @@
 
 #include "AssetTypeActions/AssetTypeActions_Enum.h"
 #include "BlueprintEditorModule.h"
+#include "AssetData.h"
 
 #define LOCTEXT_NAMESPACE "AssetTypeActions"
 
+FText FAssetTypeActions_Enum::GetAssetDescription(const FAssetData& AssetData) const
+{
+	return AssetData.GetTagValueRef<FText>(GET_MEMBER_NAME_CHECKED(UUserDefinedEnum, EnumDescription));
+}
+
 void FAssetTypeActions_Enum::OpenAssetEditor( const TArray<UObject*>& InObjects, TSharedPtr<class IToolkitHost> EditWithinLevelEditor )
 {
 	const EToolkitMode::Type Mode = EditWithinLevelEditor.IsValid() ? EToolkitMode::WorldCentric : EToolkitMode::Standalone;
diff --git a/Engine/Source/Developer/AssetTools/Private/AssetTypeActions/AssetTypeActions_Enum.h b/Engine/Source/Developer/AssetTools/Private/AssetTypeActions/AssetTypeActions_Enum.h
--- a/Engine/Source/Developer/AssetTools/Private/AssetTypeActions/AssetTypeActions_Enum.h
+++ b/Engine/Source/Developer/AssetTools/Private/AssetTypeActions/AssetTypeActions_Enum.h
@@ -12,6 +12,7 @@ class FAssetTypeActions_Enum : public FAssetTypeActions_Base
 public:
 	// IAssetTypeActions Implementation
 	virtual FText GetName() const override { return NSLOCTEXT("AssetTypeActions", "AssetTypeActions_Enum", "Enumeration"); }
+	virtual FText GetAssetDescription(const FAssetData& AssetData) const override;
 	virtual FColor GetTypeColor() const override { return FColor(255, 200, 200); }
 	virtual UClass* GetSupportedClass() const override { return UUserDefinedEnum::StaticClass(); }
 	virtual uint32 GetCategories() override { return EAssetTypeCategories::Blueprint; }
diff --git a/Engine/Source/Developer/AssetTools/Public/IAssetTools.h b/Engine/Source/Developer/AssetTools/Public/IAssetTools.h
--- a/Engine/Source/Developer/AssetTools/Public/IAssetTools.h
+++ b/Engine/Source/Developer/AssetTools/Public/IAssetTools.h
@@ -8,6 +8,7 @@
 #include "UObject/ObjectMacros.h"
 #include "UObject/Class.h"
 #include "UObject/Interface.h"
+#include "UObject/SoftObjectPath.h"
 #include "AssetTypeCategories.h"
 #include "IAssetTypeActions.h"
 #include "AutomatedAssetImportData.h"
@@ -24,30 +25,48 @@ struct FAssetRenameData
 {
 	GENERATED_BODY()
 
+	/** Object being renamed */
 	UPROPERTY(BlueprintReadWrite, Category=AssetRenameData)
 	TWeakObjectPtr<UObject> Asset;
 
+	/** New path to package without package name, ie /Game/SubDirectory */
 	UPROPERTY(BlueprintReadWrite, Category = AssetRenameData)
-	FString PackagePath;
+	FString NewPackagePath;
 
+	/** New package and asset name, new object path will be PackagePath/NewName.NewName */
 	UPROPERTY(BlueprintReadWrite, Category = AssetRenameData)
 	FString NewName;
 
-	UPROPERTY(BlueprintReadWrite, Category = AssetRenameData)
-	FString OriginalAssetPath;
+	/** Full path to old name, in form /Game/SubDirectory/OldName.OldName:SubPath*/
+	UPROPERTY()
+	FSoftObjectPath OldObjectPath;
+
+	/** New full path, may be a SubObject */
+	UPROPERTY()
+	FSoftObjectPath NewObjectPath;
+
+	/** If true, only fix soft references. This will work even if Asset is null because it has already been renamed */
+	UPROPERTY()
+	bool bOnlyFixSoftReferences;
 
 	FAssetRenameData()
+		: bOnlyFixSoftReferences(false)
 	{}
 
-	FAssetRenameData(const TWeakObjectPtr<UObject>& InAsset, const FString& InPackagePath, const FString& InNewName)
+	/** These constructors leave some fields empty, they are fixed up inside AssetRenameManager */
+	FAssetRenameData(const TWeakObjectPtr<UObject>& InAsset, const FString& InNewPackagePath, const FString& InNewName)
 		: Asset(InAsset)
-		, PackagePath(InPackagePath)
+		, NewPackagePath(InNewPackagePath)
 		, NewName(InNewName)
+		, bOnlyFixSoftReferences(false)
+	{
+	}
+	
+	FAssetRenameData(const FSoftObjectPath& InOldObjectPath, const FSoftObjectPath& InNewObjectPath, bool bInOnlyFixSoftReferences = false)
+		: OldObjectPath(InOldObjectPath)
+		, NewObjectPath(InNewObjectPath)
+		, bOnlyFixSoftReferences(bInOnlyFixSoftReferences)
 	{
-		if (InAsset.IsValid())
-		{
-			OriginalAssetPath = InAsset.Get()->GetOutermost()->GetName();
-		}
 	}
 };
 
@@ -163,6 +182,10 @@ public:
 	UFUNCTION(BlueprintCallable, Category = "Editor Scripting | Asset Tools")
 	virtual void RenameAssets(const TArray<FAssetRenameData>& AssetsAndNames) const = 0;
 
+	/** Returns list of objects that soft reference the given soft object path. This will load assets into memory to verify */
+	UFUNCTION(BlueprintCallable, Category = "Editor Scripting | Asset Tools")
+	virtual void FindSoftReferencesToObject(FSoftObjectPath TargetObject, TArray<UObject*>& ReferencingObjects) const = 0;
+
 	/** Event issued at the end of the rename process */
 	virtual FAssetPostRenameEvent& OnAssetPostRename() = 0;
 
diff --git a/Engine/Source/Developer/AutomationController/Private/AutomationControllerManager.h b/Engine/Source/Developer/AutomationController/Private/AutomationControllerManager.h
--- a/Engine/Source/Developer/AutomationController/Private/AutomationControllerManager.h
+++ b/Engine/Source/Developer/AutomationController/Private/AutomationControllerManager.h
@@ -377,17 +377,11 @@ protected:
 	 */
 	void RemoveTestRunning( const FMessageAddress& TestToRemove );
 
-	/** Resets the data holders which have been used to generate the tests available from a worker. */
-	void ResetIntermediateTestData()
-	{
-		TestInfo.Empty();
-	}
-
 	/** Changes the controller state. */
 	void SetControllerStatus( EAutomationControllerModuleState::Type AutomationTestState );
 
-	/** stores the tests that are valid for a particular device classification. */
-	void SetTestNames(const FMessageAddress& AutomationWorkerAddress);
+	/** Stores the tests that are valid for a particular device classification. */
+	void SetTestNames(const FMessageAddress& AutomationWorkerAddress, TArray<FAutomationTestInfo>& TestInfo);
 
 	/** Updates the tests to ensure they are all still running. */
 	void UpdateTests();
@@ -412,9 +406,6 @@ private:
 	/** Handles FAutomationWorkerRequestNextNetworkCommand messages. */
 	void HandleRequestNextNetworkCommandMessage( const FAutomationWorkerRequestNextNetworkCommand& Message, const IMessageContextRef& Context );
 
-	/** Handles FAutomationWorkerRequestTestsReply messages. */
-	void HandleRequestTestsReplyMessage( const FAutomationWorkerRequestTestsReply& Message, const IMessageContextRef& Context );
-
 	/** Handles FAutomationWorkerRequestTestsReplyComplete messages. */
 	void HandleRequestTestsReplyCompleteMessage(const FAutomationWorkerRequestTestsReplyComplete& Message, const IMessageContextRef& Context);
 
@@ -480,12 +471,6 @@ private:
 	/** Available stats/status for all tests. */
 	FAutomationReportManager ReportManager;
 
-	/** The number we are expecting to receive from a worker */
-	int32 NumOfTestsToReceive;
-
-	/** The collection of test data we are to send to a controller. */
-	TArray< FAutomationTestInfo > TestInfo;
-
 	/** A data holder to keep track of how long tests have been running. */
 	struct FTestRunningInfo
 	{
diff --git a/Engine/Source/Developer/AutomationController/Private/AutomationControllerManger.cpp b/Engine/Source/Developer/AutomationController/Private/AutomationControllerManger.cpp
--- a/Engine/Source/Developer/AutomationController/Private/AutomationControllerManger.cpp
+++ b/Engine/Source/Developer/AutomationController/Private/AutomationControllerManger.cpp
@@ -93,8 +93,6 @@ void FAutomationControllerManager::RequestTests()
 		{
 			FMessageAddress MessageAddress = DeviceClusterManager.GetDeviceMessageAddress(ClusterIndex, 0);
 
-			ResetIntermediateTestData();
-
 			//issue tests on appropriate platforms
 			MessageEndpoint->Send(new FAutomationWorkerRequestTests(bDeveloperDirectoryIncluded, RequestedTestFlags), MessageAddress);
 		}
@@ -218,19 +216,24 @@ void FAutomationControllerManager::ProcessComparisonQueue()
 
 			// Get the current test.
 			TSharedPtr<IAutomationReport> Report = DeviceClusterManager.GetTest(ClusterIndex, DeviceIndex);
-			check(Report.IsValid());
-
-			// Record the artifacts for the test.
-			FString ApprovedFolder = ScreenshotManager->GetLocalApprovedFolder();
-			FString UnapprovedFolder = ScreenshotManager->GetLocalUnapprovedFolder();
-			FString ComparisonFolder = ScreenshotManager->GetLocalComparisonFolder();
+			if (Report.IsValid())
+			{
+				// Record the artifacts for the test.
+				FString ApprovedFolder = ScreenshotManager->GetLocalApprovedFolder();
+				FString UnapprovedFolder = ScreenshotManager->GetLocalUnapprovedFolder();
+				FString ComparisonFolder = ScreenshotManager->GetLocalComparisonFolder();
 
-			TMap<FString, FString> LocalFiles;
-			LocalFiles.Add(TEXT("approved"), ApprovedFolder / Result.ApprovedFile);
-			LocalFiles.Add(TEXT("unapproved"), UnapprovedFolder / Result.IncomingFile);
-			LocalFiles.Add(TEXT("difference"), ComparisonFolder / Result.ComparisonFile);
+				TMap<FString, FString> LocalFiles;
+				LocalFiles.Add(TEXT("approved"), ApprovedFolder / Result.ApprovedFile);
+				LocalFiles.Add(TEXT("unapproved"), UnapprovedFolder / Result.IncomingFile);
+				LocalFiles.Add(TEXT("difference"), ComparisonFolder / Result.ComparisonFile);
 
-			Report->AddArtifact(ClusterIndex, CurrentTestPass, FAutomationArtifact(Entry->Name, EAutomationArtifactType::Comparison, LocalFiles));
+				Report->AddArtifact(ClusterIndex, CurrentTestPass, FAutomationArtifact(Entry->Name, EAutomationArtifactType::Comparison, LocalFiles));
+			}
+			else
+			{
+				UE_LOG(AutomationControllerLog, Error, TEXT("Cannot generate screenshot report for screenshot %s as report is missing"), *Result.IncomingFile);
+			}
 		}
 	}
 }
@@ -496,7 +499,6 @@ void FAutomationControllerManager::Startup()
 		.Handling<FAutomationWorkerFindWorkersResponse>(this, &FAutomationControllerManager::HandleFindWorkersResponseMessage)
 		.Handling<FAutomationWorkerPong>(this, &FAutomationControllerManager::HandlePongMessage)
 		.Handling<FAutomationWorkerRequestNextNetworkCommand>(this, &FAutomationControllerManager::HandleRequestNextNetworkCommandMessage)
-		.Handling<FAutomationWorkerRequestTestsReply>(this, &FAutomationControllerManager::HandleRequestTestsReplyMessage)
 		.Handling<FAutomationWorkerRequestTestsReplyComplete>(this, &FAutomationControllerManager::HandleRequestTestsReplyCompleteMessage)
 		.Handling<FAutomationWorkerRunTestsReply>(this, &FAutomationControllerManager::HandleRunTestsReplyMessage)
 		.Handling<FAutomationWorkerScreenImage>(this, &FAutomationControllerManager::HandleReceivedScreenShot)
@@ -513,7 +515,6 @@ void FAutomationControllerManager::Startup()
 	bDeveloperDirectoryIncluded = false;
 	RequestedTestFlags = EAutomationTestFlags::SmokeFilter | EAutomationTestFlags::EngineFilter | EAutomationTestFlags::ProductFilter | EAutomationTestFlags::PerfFilter;
 
-	NumOfTestsToReceive = 0;
 	NumTestPasses = 1;
 
 	//Default to machine name
@@ -536,7 +537,7 @@ void FAutomationControllerManager::RemoveCallbacks()
 	TestsCompleteDelegate.Clear();
 }
 
-void FAutomationControllerManager::SetTestNames(const FMessageAddress& AutomationWorkerAddress)
+void FAutomationControllerManager::SetTestNames(const FMessageAddress& AutomationWorkerAddress, TArray<FAutomationTestInfo>& TestInfo)
 {
 	int32 DeviceClusterIndex = INDEX_NONE;
 	int32 DeviceIndex = INDEX_NONE;
@@ -561,9 +562,6 @@ void FAutomationControllerManager::SetTestNames(const FMessageAddress& Automatio
 			// Ensure our test exists. If not, add it
 			ReportManager.EnsureReportExists(TestInfo[TestIndex], DeviceClusterIndex, NumTestPasses);
 		}
-
-		// Clear any intermediate data we had associated with the tests whilst building the full list of tests
-		ResetIntermediateTestData();
 	}
 	else
 	{
@@ -995,15 +993,17 @@ void FAutomationControllerManager::HandleRequestNextNetworkCommandMessage(const
 	}
 }
 
-void FAutomationControllerManager::HandleRequestTestsReplyMessage(const FAutomationWorkerRequestTestsReply& Message, const IMessageContextRef& Context)
-{
-	FAutomationTestInfo NewTest = Message.GetTestInfo();
-	TestInfo.Add(NewTest);
-}
-
 void FAutomationControllerManager::HandleRequestTestsReplyCompleteMessage(const FAutomationWorkerRequestTestsReplyComplete& Message, const IMessageContextRef& Context)
 {
-	SetTestNames(Context->GetSender());
+	TArray<FAutomationTestInfo> TestInfo;
+	TestInfo.Reset(Message.Tests.Num());
+	for (const FAutomationWorkerSingleTestReply& SingleTestReply : Message.Tests)
+	{
+		FAutomationTestInfo NewTest = SingleTestReply.GetTestInfo();
+		TestInfo.Add(NewTest);
+	}
+
+	SetTestNames(Context->GetSender(), TestInfo);
 }
 
 void FAutomationControllerManager::HandleRunTestsReplyMessage(const FAutomationWorkerRunTestsReply& Message, const IMessageContextRef& Context)
diff --git a/Engine/Source/Developer/BlueprintCompilerCppBackend/Private/BlueprintCompilerCppBackendBase.cpp b/Engine/Source/Developer/BlueprintCompilerCppBackend/Private/BlueprintCompilerCppBackendBase.cpp
--- a/Engine/Source/Developer/BlueprintCompilerCppBackend/Private/BlueprintCompilerCppBackendBase.cpp
+++ b/Engine/Source/Developer/BlueprintCompilerCppBackend/Private/BlueprintCompilerCppBackendBase.cpp
@@ -1002,6 +1002,9 @@ void FBlueprintCompilerCppBackendBase::GenerateCodeFromStruct(UUserDefinedStruct
 
 		FEmitDefaultValueHelper::GenerateGetDefaultValue(SourceStruct, EmitterContext);
 
+		EmitterContext.Header.AddLine(TEXT("static void __StaticDependenciesAssets(TArray<FBlueprintDependencyData>& AssetsToLoad);"));
+		EmitterContext.Header.AddLine(TEXT("static void __StaticDependencies_DirectlyUsedAssets(TArray<FBlueprintDependencyData>& AssetsToLoad);"));
+
 		EmitterContext.Header.AddLine(FString::Printf(TEXT("bool operator== (const %s& __Other) const"), *CppStructName));
 		EmitterContext.Header.AddLine(TEXT("{"));
 		EmitterContext.Header.IncreaseIndent();
@@ -1019,6 +1022,9 @@ void FBlueprintCompilerCppBackendBase::GenerateCodeFromStruct(UUserDefinedStruct
 		EmitterContext.Header.AddLine(TEXT("};"));
 	}
 
+	FEmitDefaultValueHelper::AddStaticFunctionsForDependencies(EmitterContext, nullptr, NativizationOptions);
+	FEmitDefaultValueHelper::AddRegisterHelper(EmitterContext);
+
 	OutCPPCode = MoveTemp(EmitterContext.Body.Result);
 	OutHeaderCode = MoveTemp(EmitterContext.Header.Result);
 }
diff --git a/Engine/Source/Developer/BlueprintCompilerCppBackend/Private/BlueprintCompilerCppBackendGatherDependencies.cpp b/Engine/Source/Developer/BlueprintCompilerCppBackend/Private/BlueprintCompilerCppBackendGatherDependencies.cpp
--- a/Engine/Source/Developer/BlueprintCompilerCppBackend/Private/BlueprintCompilerCppBackendGatherDependencies.cpp
+++ b/Engine/Source/Developer/BlueprintCompilerCppBackend/Private/BlueprintCompilerCppBackendGatherDependencies.cpp
@@ -562,9 +562,9 @@ void FGatherConvertedClassDependencies::DependenciesForHeader()
 		const bool bIsMemberVariable = OwnerProperty && (OwnerProperty->GetOuter() == OriginalStruct);
 		if (bIsParam || bIsMemberVariable)
 		{
-			if (auto AssetClassProperty = Cast<const UAssetClassProperty>(Property))
+			if (auto SoftClassProperty = Cast<const USoftClassProperty>(Property))
 			{
-				DeclareInHeader.Add(GetFirstNativeOrConvertedClass(AssetClassProperty->MetaClass));
+				DeclareInHeader.Add(GetFirstNativeOrConvertedClass(SoftClassProperty->MetaClass));
 			}
 			if (auto ClassProperty = Cast<const UClassProperty>(Property))
 			{
@@ -725,9 +725,9 @@ public:
 	TSet<UObject*> References;
 
 	//~ Begin FArchive Interface
-	virtual FArchive& operator<< (class FLazyObjectPtr& Value) override { return *this; }
-	virtual FArchive& operator<< (class FAssetPtr& Value) override { return *this; }
-	virtual FArchive& operator<< (struct FStringAssetReference& Value) override { return *this; }
+	virtual FArchive& operator<<(struct FLazyObjectPtr& Value) override { return *this; }
+	virtual FArchive& operator<<(struct FSoftObjectPtr& Value) override { return *this; }
+	virtual FArchive& operator<<(struct FSoftObjectPath& Value) override { return *this; }
 
 	virtual FArchive& operator<<(UObject*& Object) override
 	{
diff --git a/Engine/Source/Developer/BlueprintCompilerCppBackend/Private/BlueprintCompilerCppBackendUtils.cpp b/Engine/Source/Developer/BlueprintCompilerCppBackend/Private/BlueprintCompilerCppBackendUtils.cpp
--- a/Engine/Source/Developer/BlueprintCompilerCppBackend/Private/BlueprintCompilerCppBackendUtils.cpp
+++ b/Engine/Source/Developer/BlueprintCompilerCppBackend/Private/BlueprintCompilerCppBackendUtils.cpp
@@ -7,7 +7,7 @@
 #include "UObject/MetaData.h"
 #include "UObject/TextProperty.h"
 #include "UObject/PropertyPortFlags.h"
-#include "Misc/StringClassReference.h"
+#include "UObject/SoftObjectPath.h"
 #include "Components/ActorComponent.h"
 #include "Engine/BlueprintGeneratedClass.h"
 #include "Engine/UserDefinedEnum.h"
@@ -253,7 +253,7 @@ FString FEmitterLocalContext::ExportTextItem(const UProperty* Property, const vo
 			ValueStr = TEXT("(-2147483647 - 1)");
 		}
 	}
-	if (Property->IsA<UAssetObjectProperty>())
+	if (Property->IsA<USoftObjectProperty>())
 	{
 		const FString TypeText = ExportCppDeclaration(Property, EExportedDeclaration::Parameter, LocalExportCPPFlags, EPropertyNameInDeclaration::Skip);
 		return FString::Printf(TEXT("%s(%s)"), *TypeText, *ValueStr);
@@ -298,9 +298,9 @@ FString FEmitterLocalContext::ExportCppDeclaration(const UProperty* Property, EE
 	{
 		GetActualNameCPP(ClassProperty, ClassProperty->MetaClass);
 	}
-	else if (auto AssetClassProperty = Cast<const UAssetClassProperty>(Property))
+	else if (auto SoftClassProperty = Cast<const USoftClassProperty>(Property))
 	{
-		GetActualNameCPP(AssetClassProperty, AssetClassProperty->MetaClass);
+		GetActualNameCPP(SoftClassProperty, SoftClassProperty->MetaClass);
 	}
 	else if (auto ObjectProperty = Cast<const UObjectPropertyBase>(Property))
 	{
@@ -1148,7 +1148,7 @@ FString FEmitHelper::LiteralTerm(FEmitterLocalContext& EmitterContext, const FEd
 		}
 		return FString(TEXT("((UClass*)nullptr)"));
 	}
-	else if((UEdGraphSchema_K2::PC_AssetClass == Type.PinCategory) || (UEdGraphSchema_K2::PC_Asset == Type.PinCategory))
+	else if((UEdGraphSchema_K2::PC_SoftClass == Type.PinCategory) || (UEdGraphSchema_K2::PC_SoftObject == Type.PinCategory))
 	{
 		UClass* MetaClass = Cast<UClass>(Type.PinSubCategoryObject.Get());
 		MetaClass = MetaClass ? MetaClass : UObject::StaticClass();
@@ -1156,9 +1156,9 @@ FString FEmitHelper::LiteralTerm(FEmitterLocalContext& EmitterContext, const FEd
 
 		if (!CustomValue.IsEmpty())
 		{
-			const bool bAssetSubclassOf = (UEdGraphSchema_K2::PC_AssetClass == Type.PinCategory);
-			return FString::Printf(TEXT("%s<%s>(FStringAssetReference(TEXT(\"%s\")))")
-				, bAssetSubclassOf ? TEXT("TAssetSubclassOf") : TEXT("TAssetPtr")
+			const bool bAssetSubclassOf = (UEdGraphSchema_K2::PC_SoftClass == Type.PinCategory);
+			return FString::Printf(TEXT("%s<%s>(FSoftObjectPath(TEXT(\"%s\")))")
+				, bAssetSubclassOf ? TEXT("TSoftClassPtr") : TEXT("TSoftObjectPtr")
 				, *ObjTypeStr
 				, *(CustomValue.ReplaceCharWithEscapedChar()));
 		}
@@ -1261,11 +1261,11 @@ FString FEmitHelper::PinTypeToNativeType(const FEdGraphPinType& Type)
 				return FString::Printf(TEXT("TSubclassOf<%s>"), *FEmitHelper::GetCppName(Class));
 			}
 		}
-		else if (UEdGraphSchema_K2::PC_AssetClass == InType.PinCategory)
+		else if (UEdGraphSchema_K2::PC_SoftClass == InType.PinCategory)
 		{
 			if (UClass* Class = Cast<UClass>(InType.PinSubCategoryObject.Get()))
 			{
-				return FString::Printf(TEXT("TAssetSubclassOf<%s>"), *FEmitHelper::GetCppName(Class));
+				return FString::Printf(TEXT("TSoftClassPtr<%s>"), *FEmitHelper::GetCppName(Class));
 			}
 		}
 		else if (UEdGraphSchema_K2::PC_Interface == InType.PinCategory)
@@ -1275,11 +1275,11 @@ FString FEmitHelper::PinTypeToNativeType(const FEdGraphPinType& Type)
 				return FString::Printf(TEXT("TScriptInterface<%s>"), *FEmitHelper::GetCppName(Class));
 			}
 		}
-		else if (UEdGraphSchema_K2::PC_Asset == InType.PinCategory)
+		else if (UEdGraphSchema_K2::PC_SoftObject == InType.PinCategory)
 		{
 			if (UClass* Class = Cast<UClass>(InType.PinSubCategoryObject.Get()))
 			{
-				return FString::Printf(TEXT("TAssetPtr<%s>"), *FEmitHelper::GetCppName(Class));
+				return FString::Printf(TEXT("TSoftObjectPtr<%s>"), *FEmitHelper::GetCppName(Class));
 			}
 		}
 		else if (UEdGraphSchema_K2::PC_Object == InType.PinCategory)
@@ -1573,7 +1573,7 @@ void FNativizationSummaryHelper::InaccessibleProperty(const UProperty* Property)
 	TSharedPtr<FNativizationSummary> NativizationSummary = BackEndModule.NativizationSummary();
 	if (NativizationSummary.IsValid())
 	{
-		FStringAssetReference Key(Property);
+		FSoftObjectPath Key(Property);
 		int32* FoundStat = NativizationSummary->InaccessiblePropertyStat.Find(Key);
 		if (FoundStat)
 		{
@@ -1599,7 +1599,7 @@ static void FNativizationSummaryHelper__MemberUsed(const UClass* Class, const UF
 			const bool bUnrelatedClass = !Class->IsChildOf(Owner);
 			if (AnimBP && bUnrelatedClass)
 			{
-				FNativizationSummary::FAnimBlueprintDetails& AnimBlueprintDetails = NativizationSummary->AnimBlueprintStat.FindOrAdd(FStringAssetReference(AnimBP));
+				FNativizationSummary::FAnimBlueprintDetails& AnimBlueprintDetails = NativizationSummary->AnimBlueprintStat.FindOrAdd(FSoftObjectPath(AnimBP));
 				(AnimBlueprintDetails.*CouterPtr)++;
 			}
 		}
@@ -1625,19 +1625,19 @@ void FNativizationSummaryHelper::ReducibleFunciton(const UClass* OriginalClass)
 		UAnimBlueprint* AnimBP = Cast<UAnimBlueprint>(UBlueprint::GetBlueprintFromClass(OriginalClass));
 		if (NativizationSummary.IsValid() && OriginalClass && AnimBP)
 		{
-			FNativizationSummary::FAnimBlueprintDetails& AnimBlueprintDetails = NativizationSummary->AnimBlueprintStat.FindOrAdd(FStringAssetReference(AnimBP));
+			FNativizationSummary::FAnimBlueprintDetails& AnimBlueprintDetails = NativizationSummary->AnimBlueprintStat.FindOrAdd(FSoftObjectPath(AnimBP));
 			AnimBlueprintDetails.ReducibleFunctions++;
 		}
 	}
 }
 
-void FNativizationSummaryHelper::RegisterRequiredModules(const FName PlatformName, const TSet<TAssetPtr<UPackage>>& InModules)
+void FNativizationSummaryHelper::RegisterRequiredModules(const FName PlatformName, const TSet<TSoftObjectPtr<UPackage>>& InModules)
 {
 	IBlueprintCompilerCppBackendModule& BackEndModule = (IBlueprintCompilerCppBackendModule&)IBlueprintCompilerCppBackendModule::Get();
 	TSharedPtr<FNativizationSummary> NativizationSummary = BackEndModule.NativizationSummary();
 	if (NativizationSummary.IsValid())
 	{
-		TSet<TAssetPtr<UPackage>>& Modules = NativizationSummary->ModulesRequiredByPlatform.FindOrAdd(PlatformName);
+		TSet<TSoftObjectPtr<UPackage>>& Modules = NativizationSummary->ModulesRequiredByPlatform.FindOrAdd(PlatformName);
 		Modules.Append(InModules);
 	}
 }
@@ -1650,7 +1650,7 @@ void FNativizationSummaryHelper::RegisterClass(const UClass* OriginalClass)
 	if (NativizationSummary.IsValid() && OriginalClass && AnimBP)
 	{
 		{
-			FNativizationSummary::FAnimBlueprintDetails& AnimBlueprintDetails = NativizationSummary->AnimBlueprintStat.FindOrAdd(FStringAssetReference(AnimBP));
+			FNativizationSummary::FAnimBlueprintDetails& AnimBlueprintDetails = NativizationSummary->AnimBlueprintStat.FindOrAdd(FSoftObjectPath(AnimBP));
 
 			AnimBlueprintDetails.Variables = AnimBP->NewVariables.Num();
 
@@ -1669,7 +1669,7 @@ void FNativizationSummaryHelper::RegisterClass(const UClass* OriginalClass)
 			UAnimBlueprint* ParentAnimBP = Cast<UAnimBlueprint>(UBlueprint::GetBlueprintFromClass(SuperClass));
 			if (ParentAnimBP)
 			{
-				FNativizationSummary::FAnimBlueprintDetails& AnimBlueprintDetails = NativizationSummary->AnimBlueprintStat.FindOrAdd(FStringAssetReference(ParentAnimBP));
+				FNativizationSummary::FAnimBlueprintDetails& AnimBlueprintDetails = NativizationSummary->AnimBlueprintStat.FindOrAdd(FSoftObjectPath(ParentAnimBP));
 				AnimBlueprintDetails.Children++;
 			}
 		}
@@ -1757,7 +1757,7 @@ FString FEmitHelper::AccessInaccessibleProperty(FEmitterLocalContext& EmitterCon
 
 struct FSearchableValuesdHelper_StaticData
 {
-	TArray<FStringClassReference> ClassesWithStaticSearchableValues;
+	TArray<FSoftClassPath> ClassesWithStaticSearchableValues;
 	TArray<FName> TagPropertyNames;
 private:
 	FSearchableValuesdHelper_StaticData()
@@ -1767,7 +1767,7 @@ private:
 			GConfig->GetArray(TEXT("BlueprintNativizationSettings"), TEXT("ClassesWithStaticSearchableValues"), Paths, GEditorIni);
 			for (FString& Path : Paths)
 			{
-				ClassesWithStaticSearchableValues.Add(FStringClassReference(Path));
+				ClassesWithStaticSearchableValues.Add(FSoftClassPath(Path));
 			}
 		}
 
@@ -1791,7 +1791,7 @@ public:
 
 bool FBackendHelperStaticSearchableValues::HasSearchableValues(UClass* InClass)
 {
-	for (const FStringClassReference& ClassStrRef : FSearchableValuesdHelper_StaticData::Get().ClassesWithStaticSearchableValues)
+	for (const FSoftClassPath& ClassStrRef : FSearchableValuesdHelper_StaticData::Get().ClassesWithStaticSearchableValues)
 	{
 		UClass* IterClass = ClassStrRef.ResolveClass();
 		if (IterClass && InClass && InClass->IsChildOf(IterClass))
@@ -1919,7 +1919,7 @@ FString FDependenciesGlobalMapHelper::EmitBodyCode(const FString& PCHFilename)
 	return CodeText.Result;
 }
 
-FNativizationSummary::FDependencyRecord& FDependenciesGlobalMapHelper::FindDependencyRecord(const FStringAssetReference& Key)
+FNativizationSummary::FDependencyRecord& FDependenciesGlobalMapHelper::FindDependencyRecord(const FSoftObjectPath& Key)
 {
 	auto& DependenciesGlobalMap = GetDependenciesGlobalMap();
 	FNativizationSummary::FDependencyRecord& DependencyRecord = DependenciesGlobalMap.FindOrAdd(Key);
@@ -1930,7 +1930,7 @@ FNativizationSummary::FDependencyRecord& FDependenciesGlobalMapHelper::FindDepen
 	return DependencyRecord;
 }
 
-TMap<FStringAssetReference, FNativizationSummary::FDependencyRecord>& FDependenciesGlobalMapHelper::GetDependenciesGlobalMap()
+TMap<FSoftObjectPath, FNativizationSummary::FDependencyRecord>& FDependenciesGlobalMapHelper::GetDependenciesGlobalMap()
 {
 	IBlueprintCompilerCppBackendModule& BackEndModule = (IBlueprintCompilerCppBackendModule&)IBlueprintCompilerCppBackendModule::Get();
 	TSharedPtr<FNativizationSummary> NativizationSummary = BackEndModule.NativizationSummary();
@@ -1990,7 +1990,7 @@ struct FStructAccessHelper_StaticData
 {
 	TMap<const UStruct*, FString> BaseStructureAccessorsMap;
 	TMap<const UStruct*, bool> SupportsDirectNativeAccessMap;
-	TArray<FStringClassReference> NoExportTypesWithDirectNativeFieldAccess;
+	TArray<FSoftClassPath> NoExportTypesWithDirectNativeFieldAccess;
 
 	static FStructAccessHelper_StaticData& Get()
 	{
@@ -2026,7 +2026,7 @@ private:
 			GConfig->GetArray(TEXT("BlueprintNativizationSettings"), TEXT("NoExportTypesWithDirectNativeFieldAccess"), Paths, GEditorIni);
 			for (FString& Path : Paths)
 			{
-				NoExportTypesWithDirectNativeFieldAccess.Add(FStringClassReference(Path));
+				NoExportTypesWithDirectNativeFieldAccess.Add(FSoftClassPath(Path));
 			}
 		}
 	}
diff --git a/Engine/Source/Developer/BlueprintCompilerCppBackend/Private/BlueprintCompilerCppBackendUtils.h b/Engine/Source/Developer/BlueprintCompilerCppBackend/Private/BlueprintCompilerCppBackendUtils.h
--- a/Engine/Source/Developer/BlueprintCompilerCppBackend/Private/BlueprintCompilerCppBackendUtils.h
+++ b/Engine/Source/Developer/BlueprintCompilerCppBackend/Private/BlueprintCompilerCppBackendUtils.h
@@ -3,7 +3,7 @@
 
 #include "CoreMinimal.h"
 #include "UObject/Class.h"
-#include "Misc/StringAssetReference.h"
+#include "UObject/SoftObjectPath.h"
 #include "UObject/UnrealType.h"
 #include "IBlueprintCompilerCppBackendModule.h"
 #include "BlueprintCompilerCppBackendGatherDependencies.h"
@@ -408,17 +408,17 @@ struct FNativizationSummaryHelper
 
 	static void ReducibleFunciton(const UClass* OriginalClass);
 
-	static void RegisterRequiredModules(const FName PlatformName, const TSet<TAssetPtr<UPackage>>& Modules);
+	static void RegisterRequiredModules(const FName PlatformName, const TSet<TSoftObjectPtr<UPackage>>& Modules);
 };
 struct FDependenciesGlobalMapHelper
 {
 	static FString EmitHeaderCode();
 	static FString EmitBodyCode(const FString& PCHFilename);
 
-	static FNativizationSummary::FDependencyRecord& FindDependencyRecord(const FStringAssetReference& Key);
+	static FNativizationSummary::FDependencyRecord& FindDependencyRecord(const FSoftObjectPath& Key);
 
 private:
-	static TMap<FStringAssetReference, FNativizationSummary::FDependencyRecord>& GetDependenciesGlobalMap();
+	static TMap<FSoftObjectPath, FNativizationSummary::FDependencyRecord>& GetDependenciesGlobalMap();
 };
 
 struct FDisableUnwantedWarningOnScope
diff --git a/Engine/Source/Developer/BlueprintCompilerCppBackend/Private/BlueprintCompilerCppBackendValueHelper.cpp b/Engine/Source/Developer/BlueprintCompilerCppBackend/Private/BlueprintCompilerCppBackendValueHelper.cpp
--- a/Engine/Source/Developer/BlueprintCompilerCppBackend/Private/BlueprintCompilerCppBackendValueHelper.cpp
+++ b/Engine/Source/Developer/BlueprintCompilerCppBackend/Private/BlueprintCompilerCppBackendValueHelper.cpp
@@ -1071,26 +1071,22 @@ public:
 
 struct FFakeImportTableHelper
 {
-	TSet<UObject*> SerializeBeforeSerializeClassDependencies;
+	TSet<UObject*> SerializeBeforeSerializeStructDependencies;
 
 	TSet<UObject*> SerializeBeforeCreateCDODependencies;
 
-	FFakeImportTableHelper(UClass* SourceClass, UClass* OriginalClass, FEmitterLocalContext& Context)
+	FFakeImportTableHelper(UStruct* SourceStruct, UClass* OriginalClass, FEmitterLocalContext& Context)
 	{
-		if (ensure(SourceClass) && ensure(OriginalClass))
+		UClass* SourceClass = Cast<UClass>(SourceStruct);
+		if (ensure(SourceStruct) && ensure(!SourceClass || OriginalClass))
 		{
-			auto GatherDependencies = [&](UClass* InClass)
+			auto GatherDependencies = [&](UStruct* InStruct)
 			{
-				SerializeBeforeSerializeClassDependencies.Add(InClass->GetSuperClass());
+				SerializeBeforeSerializeStructDependencies.Add(InStruct->GetSuperStruct());
 
-				for (const FImplementedInterface& ImplementedInterface : InClass->Interfaces)
-				{
-					SerializeBeforeSerializeClassDependencies.Add(ImplementedInterface.Class);
-				}
-
-				TArray<UObject*> ObjectsInsideClass;
-				GetObjectsWithOuter(InClass, ObjectsInsideClass, true);
-				for (UObject* Obj : ObjectsInsideClass)
+				TArray<UObject*> ObjectsInsideStruct;
+				GetObjectsWithOuter(InStruct, ObjectsInsideStruct, true);
+				for (UObject* Obj : ObjectsInsideStruct)
 				{
 					UProperty* Property = Cast<UProperty>(Obj);
 					if (!Property)
@@ -1107,28 +1103,40 @@ struct FFakeImportTableHelper
 					// Let UDS_A contain UDS_B. Let UDS_B contain an array or a set of UDS_A. It causes a cyclic dependency. 
 					// Should we try to fix it at this stage?
 
-					const bool bIsParam = (0 != (OwnerProperty->PropertyFlags & CPF_Parm)) && OwnerProperty->IsIn(InClass);
-					const bool bIsMemberVariable = (OwnerProperty->GetOuter() == InClass);
+					const bool bIsParam = (0 != (OwnerProperty->PropertyFlags & CPF_Parm)) && OwnerProperty->IsIn(InStruct);
+					const bool bIsMemberVariable = (OwnerProperty->GetOuter() == InStruct);
 					if (bIsParam || bIsMemberVariable) // Affects the class signature. It is necessary while ZCOnstructor/linking.
 					{
 						TArray<UObject*> LocalPreloadDependencies;
 						Property->GetPreloadDependencies(LocalPreloadDependencies);
 						for (UObject* Dependency : LocalPreloadDependencies)
 						{
-							const bool bDependencyMustBeSerializedBeforeClassIsLinked = Dependency
+							const bool bDependencyMustBeSerializedBeforeStructIsLinked = Dependency
 								&& (Dependency->IsA<UScriptStruct>() || Dependency->IsA<UEnum>());
-							if (bDependencyMustBeSerializedBeforeClassIsLinked)
+							if (bDependencyMustBeSerializedBeforeStructIsLinked)
 							{
-								SerializeBeforeSerializeClassDependencies.Add(Dependency);
+								SerializeBeforeSerializeStructDependencies.Add(Dependency);
 							}
 						}
 					}
 				}
-				SerializeBeforeCreateCDODependencies.Add(InClass->GetSuperClass()->GetDefaultObject());
+
+				if (UClass* Class = Cast<UClass>(InStruct))
+				{
+					for (const FImplementedInterface& ImplementedInterface : Class->Interfaces)
+					{
+						SerializeBeforeSerializeStructDependencies.Add(ImplementedInterface.Class);
+					}
+
+					SerializeBeforeCreateCDODependencies.Add(Class->GetSuperClass()->GetDefaultObject());
+				}
 			};
 
-			GatherDependencies(SourceClass);
-			GatherDependencies(OriginalClass);
+			GatherDependencies(SourceStruct);
+			if (OriginalClass)
+			{
+				GatherDependencies(OriginalClass);
+			}
 
 			auto GetClassesOfSubobjects = [&](TMap<UObject*, FString>& SubobjectsMap)
 			{
@@ -1143,6 +1151,7 @@ struct FFakeImportTableHelper
 					}
 				}
 			};
+
 			GetClassesOfSubobjects(Context.ClassSubobjectsMap);
 			GetClassesOfSubobjects(Context.CommonSubobjectsMap);
 		}
@@ -1154,16 +1163,16 @@ struct FFakeImportTableHelper
 
 		{
 			//Dynamic Class requires no non-native class, owner, archetype..
-			CompactDataRef.ClassDependency.bSerializationBeforeCreateDependency = false;
-			CompactDataRef.ClassDependency.bCreateBeforeCreateDependency = false;
+			CompactDataRef.StructDependency.bSerializationBeforeCreateDependency = false;
+			CompactDataRef.StructDependency.bCreateBeforeCreateDependency = false;
 
-			const bool bDependencyNecessaryForLinking = SerializeBeforeSerializeClassDependencies.Contains(const_cast<UObject*>(Asset));
+			const bool bDependencyNecessaryForLinking = SerializeBeforeSerializeStructDependencies.Contains(const_cast<UObject*>(Asset));
 
 			// Super Class, Interfaces, ScriptStructs, Enums..
-			CompactDataRef.ClassDependency.bSerializationBeforeSerializationDependency = bDependencyNecessaryForLinking;
+			CompactDataRef.StructDependency.bSerializationBeforeSerializationDependency = bDependencyNecessaryForLinking;
 
 			// Everything else
-			CompactDataRef.ClassDependency.bCreateBeforeSerializationDependency = !bDependencyNecessaryForLinking;
+			CompactDataRef.StructDependency.bCreateBeforeSerializationDependency = !bDependencyNecessaryForLinking;
 		}
 
 		{
@@ -1213,10 +1222,14 @@ void FEmitDefaultValueHelper::AddStaticFunctionsForDependencies(FEmitterLocalCon
 	}
 
 	// HELPERS
-	auto SourceClass = Context.GetCurrentlyGeneratedClass();
-	auto OriginalClass = Context.Dependencies.FindOriginalClass(SourceClass);
-	const FString CppClassName = FEmitHelper::GetCppName(SourceClass);
-	FFakeImportTableHelper FakeImportTableHelper(SourceClass, OriginalClass, Context);
+	UStruct* SourceStruct = Context.Dependencies.GetActualStruct();
+	UClass* OriginalClass = nullptr;
+	if (UClass* SourceClass = Cast<UClass>(SourceStruct))
+	{
+		OriginalClass = Context.Dependencies.FindOriginalClass(SourceClass);
+	}
+	const FString CppTypeName = FEmitHelper::GetCppName(SourceStruct);
+	FFakeImportTableHelper FakeImportTableHelper(SourceStruct, OriginalClass, Context);
 
 	auto CreateAssetToLoadString = [&](const UObject* AssetObj) -> FString
 	{
@@ -1249,7 +1262,7 @@ void FEmitDefaultValueHelper::AddStaticFunctionsForDependencies(FEmitterLocalCon
 		if (InAsset && IsEditorOnlyObject(InAsset))
 		{
 			UE_LOG(LogK2Compiler, Warning, TEXT("Nativized %d depends on editor only asset: %s")
-				, (OriginalClass ? *OriginalClass->GetPathName() : *CppClassName)
+				, (OriginalClass ? *OriginalClass->GetPathName() : *CppTypeName)
 				,*InAsset->GetPathName());
 			OptionalComment = TEXT("Editor Only asset");
 			return FCompactBlueprintDependencyData{};
@@ -1311,34 +1324,35 @@ void FEmitDefaultValueHelper::AddStaticFunctionsForDependencies(FEmitterLocalCon
 		{
 			FString OptionalComment;
 			const FCompactBlueprintDependencyData DependencyRecord = CreateDependencyRecord(LocAsset, OptionalComment);
-			Context.AddLine(FString::Printf(TEXT("{%d, %s, %s},  // %s %s ")
-				, DependencyRecord.ObjectRefIndex
-				, *BlueprintDependencyTypeToString(DependencyRecord.ClassDependency)
-				, *BlueprintDependencyTypeToString(DependencyRecord.CDODependency)
-				, *OptionalComment
-				, *LocAsset->GetFullName()));
+
+			if (SourceStruct->IsA<UClass>())
+			{
+				Context.AddLine(FString::Printf(TEXT("{%d, %s, %s},  // %s %s ")
+					, DependencyRecord.ObjectRefIndex
+					, *BlueprintDependencyTypeToString(DependencyRecord.StructDependency)
+					, *BlueprintDependencyTypeToString(DependencyRecord.CDODependency)
+					, *OptionalComment
+					, *LocAsset->GetFullName()));
+			}
+			else
+			{
+				Context.AddLine(FString::Printf(TEXT("{%d, %s},  // %s %s ")
+					, DependencyRecord.ObjectRefIndex
+					, *BlueprintDependencyTypeToString(DependencyRecord.StructDependency)
+					, *OptionalComment
+					, *LocAsset->GetFullName()));
+			}
 		}
 
 		if (Assets.Num())
 		{
 			Context.DecreaseIndent();
 			Context.AddLine(TEXT("};"));
-			Context.AddLine(TEXT("for(const FCompactBlueprintDependencyData CompactData : LocCompactBlueprintDependencyData)"));
+			Context.AddLine(TEXT("for(const FCompactBlueprintDependencyData& CompactData : LocCompactBlueprintDependencyData)"));
 			Context.AddLine(TEXT("{"));
-			if (bEnableBootTimeEDLOptimization)
-			{
-				Context.AddLine(TEXT("\tAssetsToLoad.Add(FBlueprintDependencyData("));
-			}
-			else
-			{
-				Context.AddLine(TEXT("\tAssetsToLoad.AddUnique(FBlueprintDependencyData("));
-			}
-			Context.AddLine(TEXT("\t\tF__NativeDependencies::Get(CompactData.ObjectRefIndex)"));
-			Context.AddLine(TEXT("\t\t,CompactData.ClassDependency"));
-			Context.AddLine(TEXT("\t\t,CompactData.CDODependency"));
-			Context.AddLine(TEXT("\t\t,CompactData.ObjectRefIndex"));
-			Context.AddLine(TEXT("\t));"));
-			Context.AddLine(TEXT("} "));
+			Context.AddLine(FString::Printf(TEXT("\tAssetsToLoad.%s(FBlueprintDependencyData(F__NativeDependencies::Get(CompactData.ObjectRefIndex), CompactData));")
+				, bEnableBootTimeEDLOptimization ? TEXT("Add") : TEXT("AddUnique")));
+			Context.AddLine(TEXT("}"));
 		}
 	};
 
@@ -1360,8 +1374,9 @@ void FEmitDefaultValueHelper::AddStaticFunctionsForDependencies(FEmitterLocalCon
 
 
 	// 3. LIST OF UsedAssets
+	if (SourceStruct->IsA<UClass>())
 	{
-		Context.AddLine(FString::Printf(TEXT("void %s::__StaticDependencies_DirectlyUsedAssets(TArray<FBlueprintDependencyData>& AssetsToLoad)"), *CppClassName));
+		Context.AddLine(FString::Printf(TEXT("void %s::__StaticDependencies_DirectlyUsedAssets(TArray<FBlueprintDependencyData>& AssetsToLoad)"), *CppTypeName));
 		Context.AddLine(TEXT("{"));
 		Context.IncreaseIndent();
 		TArray<const UObject*> AssetsToAdd;
@@ -1379,45 +1394,48 @@ void FEmitDefaultValueHelper::AddStaticFunctionsForDependencies(FEmitterLocalCon
 
 	// 4. REMAINING DEPENDENCIES
 	{
-		Context.AddLine(FString::Printf(TEXT("void %s::__StaticDependenciesAssets(TArray<FBlueprintDependencyData>& AssetsToLoad)"), *CppClassName));
+		Context.AddLine(FString::Printf(TEXT("void %s::__StaticDependenciesAssets(TArray<FBlueprintDependencyData>& AssetsToLoad)"), *CppTypeName));
 		Context.AddLine(TEXT("{"));
 		Context.IncreaseIndent();
 
-		if (!OtherBPGCs.Num() || bEnableBootTimeEDLOptimization)
+		if (SourceStruct->IsA<UClass>())
 		{
-			Context.AddLine(FString(TEXT("__StaticDependencies_DirectlyUsedAssets(AssetsToLoad);")));
-		}
-		else
-		{
-			// To reduce the size of __StaticDependenciesAssets, all __StaticDependenciesAssets of listed BPs will be called.
-			FNativizationSummary::FDependencyRecord& DependencyRecord = FDependenciesGlobalMapHelper::FindDependencyRecord(OriginalClass);
-			ensure(DependencyRecord.Index >= 0); 
-			if (DependencyRecord.NativeLine.IsEmpty())
+			if (!OtherBPGCs.Num() || bEnableBootTimeEDLOptimization)
 			{
-				DependencyRecord.NativeLine = CreateAssetToLoadString(OriginalClass);
+				Context.AddLine(FString(TEXT("__StaticDependencies_DirectlyUsedAssets(AssetsToLoad);")));
 			}
-			Context.AddLine(FString::Printf(TEXT("const int16 __OwnIndex = %d;"), DependencyRecord.Index));
-			Context.AddLine(FString(TEXT("if(FBlueprintDependencyData::ContainsDependencyData(AssetsToLoad, __OwnIndex)) { return; }")));
-			Context.AddLine(TEXT("if(GEventDrivenLoaderEnabled && EVENT_DRIVEN_ASYNC_LOAD_ACTIVE_AT_RUNTIME){ __StaticDependencies_DirectlyUsedAssets(AssetsToLoad); }"));
-			Context.AddLine(TEXT("else"));
-			Context.AddLine(TEXT("{"));
-			Context.IncreaseIndent();
-			Context.AddLine(FString(TEXT("const bool __FirstFunctionCall = !AssetsToLoad.Num();")));
-			Context.AddLine(FString(TEXT("TArray<FBlueprintDependencyData> Temp;")));
-			// Other __StaticDependenciesAssets fucntions should not see the assets added by __StaticDependencies_DirectlyUsedAssets
-			// But in the first function called the assets from __StaticDependencies_DirectlyUsedAssets must go first in unchanged order (to satisfy FConvertedBlueprintsDependencies::FillUsedAssetsInDynamicClass)
-			Context.AddLine(FString(TEXT("__StaticDependencies_DirectlyUsedAssets(__FirstFunctionCall ? AssetsToLoad : Temp);")));
-			Context.AddLine(FString(TEXT("TArray<FBlueprintDependencyData>& ArrayUnaffectedByDirectlyUsedAssets = __FirstFunctionCall ? Temp : AssetsToLoad;")));
-
-			Context.AddLine(FString(TEXT("ArrayUnaffectedByDirectlyUsedAssets.AddUnique(FBlueprintDependencyData(F__NativeDependencies::Get(__OwnIndex), {}, {}, __OwnIndex));")));
-
-			for (const UBlueprintGeneratedClass* OtherBPGC : OtherBPGCs)
+			else
 			{
-				Context.AddLine(FString::Printf(TEXT("%s::__StaticDependenciesAssets(ArrayUnaffectedByDirectlyUsedAssets);"), *FEmitHelper::GetCppName(OtherBPGC)));
+				// To reduce the size of __StaticDependenciesAssets, all __StaticDependenciesAssets of listed BPs will be called.
+				FNativizationSummary::FDependencyRecord& DependencyRecord = FDependenciesGlobalMapHelper::FindDependencyRecord(OriginalClass);
+				ensure(DependencyRecord.Index >= 0);
+				if (DependencyRecord.NativeLine.IsEmpty())
+				{
+					DependencyRecord.NativeLine = CreateAssetToLoadString(OriginalClass);
+				}
+				Context.AddLine(FString::Printf(TEXT("const int16 __OwnIndex = %d;"), DependencyRecord.Index));
+				Context.AddLine(FString(TEXT("if(FBlueprintDependencyData::ContainsDependencyData(AssetsToLoad, __OwnIndex)) { return; }")));
+				Context.AddLine(TEXT("if(GEventDrivenLoaderEnabled && EVENT_DRIVEN_ASYNC_LOAD_ACTIVE_AT_RUNTIME){ __StaticDependencies_DirectlyUsedAssets(AssetsToLoad); }"));
+				Context.AddLine(TEXT("else"));
+				Context.AddLine(TEXT("{"));
+				Context.IncreaseIndent();
+				Context.AddLine(FString(TEXT("const bool __FirstFunctionCall = !AssetsToLoad.Num();")));
+				Context.AddLine(FString(TEXT("TArray<FBlueprintDependencyData> Temp;")));
+				// Other __StaticDependenciesAssets fucntions should not see the assets added by __StaticDependencies_DirectlyUsedAssets
+				// But in the first function called the assets from __StaticDependencies_DirectlyUsedAssets must go first in unchanged order (to satisfy FConvertedBlueprintsDependencies::FillUsedAssetsInDynamicClass)
+				Context.AddLine(FString(TEXT("__StaticDependencies_DirectlyUsedAssets(__FirstFunctionCall ? AssetsToLoad : Temp);")));
+				Context.AddLine(FString(TEXT("TArray<FBlueprintDependencyData>& ArrayUnaffectedByDirectlyUsedAssets = __FirstFunctionCall ? Temp : AssetsToLoad;")));
+
+				Context.AddLine(FString(TEXT("ArrayUnaffectedByDirectlyUsedAssets.AddUnique(FBlueprintDependencyData(F__NativeDependencies::Get(__OwnIndex), {}, {}, __OwnIndex));")));
+
+				for (const UBlueprintGeneratedClass* OtherBPGC : OtherBPGCs)
+				{
+					Context.AddLine(FString::Printf(TEXT("%s::__StaticDependenciesAssets(ArrayUnaffectedByDirectlyUsedAssets);"), *FEmitHelper::GetCppName(OtherBPGC)));
+				}
+				Context.AddLine(FString(TEXT("FBlueprintDependencyData::AppendUniquely(AssetsToLoad, Temp);")));
+				Context.DecreaseIndent();
+				Context.AddLine(TEXT("}"));
 			}
-			Context.AddLine(FString(TEXT("FBlueprintDependencyData::AppendUniquely(AssetsToLoad, Temp);")));
-			Context.DecreaseIndent();
-			Context.AddLine(TEXT("}"));
 		}
 
 		if (bEnableBootTimeEDLOptimization)
@@ -1456,11 +1474,15 @@ void FEmitDefaultValueHelper::AddStaticFunctionsForDependencies(FEmitterLocalCon
 
 void FEmitDefaultValueHelper::AddRegisterHelper(FEmitterLocalContext& Context)
 {
-	auto SourceClass = Context.GetCurrentlyGeneratedClass();
-	auto OriginalClass = Context.Dependencies.FindOriginalClass(SourceClass);
-	const FString CppClassName = FEmitHelper::GetCppName(SourceClass);
+	UStruct* SourceStruct = Context.Dependencies.GetActualStruct();
+	const FString CppTypeName = FEmitHelper::GetCppName(SourceStruct);
+
+	if (UClass* SourceClass = Cast<UClass>(SourceStruct))
+	{
+		SourceStruct = Context.Dependencies.FindOriginalClass(SourceClass);
+	}
 
-	const FString RegisterHelperName = FString::Printf(TEXT("FRegisterHelper__%s"), *CppClassName);
+	const FString RegisterHelperName = FString::Printf(TEXT("FRegisterHelper__%s"), *CppTypeName);
 	Context.AddLine(FString::Printf(TEXT("struct %s"), *RegisterHelperName));
 	Context.AddLine(TEXT("{"));
 	Context.IncreaseIndent();
@@ -1471,8 +1493,8 @@ void FEmitDefaultValueHelper::AddRegisterHelper(FEmitterLocalContext& Context)
 
 	Context.AddLine(FString::Printf(
 		TEXT("FConvertedBlueprintsDependencies::Get().RegisterConvertedClass(TEXT(\"%s\"), &%s::__StaticDependenciesAssets);")
-		, *OriginalClass->GetOutermost()->GetPathName()
-		, *CppClassName));
+		, *SourceStruct->GetOutermost()->GetPathName()
+		, *CppTypeName));
 
 	Context.DecreaseIndent();
 	Context.AddLine(TEXT("}"));
diff --git a/Engine/Source/Developer/BlueprintCompilerCppBackend/Public/BlueprintCompilerCppBackendGatherDependencies.h b/Engine/Source/Developer/BlueprintCompilerCppBackend/Public/BlueprintCompilerCppBackendGatherDependencies.h
--- a/Engine/Source/Developer/BlueprintCompilerCppBackend/Public/BlueprintCompilerCppBackendGatherDependencies.h
+++ b/Engine/Source/Developer/BlueprintCompilerCppBackend/Public/BlueprintCompilerCppBackendGatherDependencies.h
@@ -3,7 +3,7 @@
 #pragma  once
 
 #include "CoreMinimal.h"
-#include "UObject/AssetPtr.h"
+#include "UObject/SoftObjectPtr.h"
 #include "Engine/Blueprint.h"
 
 class UBlueprintGeneratedClass;
@@ -29,7 +29,7 @@ public:
 	TSet<UField*> DeclareInHeader;
 	TSet<UField*> IncludeInBody;
 
-	TSet<TAssetPtr<UPackage>> RequiredModuleNames;
+	TSet<TSoftObjectPtr<UPackage>> RequiredModuleNames;
 
 	FCompilerNativizationOptions NativizationOptions;
 
diff --git a/Engine/Source/Developer/BlueprintCompilerCppBackend/Public/IBlueprintCompilerCppBackendModule.h b/Engine/Source/Developer/BlueprintCompilerCppBackend/Public/IBlueprintCompilerCppBackendModule.h
--- a/Engine/Source/Developer/BlueprintCompilerCppBackend/Public/IBlueprintCompilerCppBackendModule.h
+++ b/Engine/Source/Developer/BlueprintCompilerCppBackend/Public/IBlueprintCompilerCppBackendModule.h
@@ -3,7 +3,7 @@
 #pragma once
 
 #include "CoreMinimal.h"
-#include "Misc/StringAssetReference.h"
+#include "UObject/SoftObjectPath.h"
 #include "BlueprintCompilerCppBackendInterface.h"
 
 class UBlueprint;
@@ -11,7 +11,7 @@ struct FCompilerNativizationOptions;
 
 struct FNativizationSummary
 {
-	TMap<FStringAssetReference, int32> InaccessiblePropertyStat;
+	TMap<FSoftObjectPath, int32> InaccessiblePropertyStat;
 
 	struct FAnimBlueprintDetails
 	{
@@ -27,7 +27,7 @@ struct FNativizationSummary
 		{}
 	};
 
-	TMap<FStringAssetReference, FAnimBlueprintDetails> AnimBlueprintStat;
+	TMap<FSoftObjectPath, FAnimBlueprintDetails> AnimBlueprintStat;
 
 	int32 MemberVariablesFromGraph;
 
@@ -40,9 +40,9 @@ struct FNativizationSummary
 		FDependencyRecord() : Index(-1) {}
 	};
 
-	TMap<FStringAssetReference, FDependencyRecord> DependenciesGlobalMap;
+	TMap<FSoftObjectPath, FDependencyRecord> DependenciesGlobalMap;
 
-	TMap<FName, TSet<TAssetPtr<UPackage>>> ModulesRequiredByPlatform;
+	TMap<FName, TSet<TSoftObjectPtr<UPackage>>> ModulesRequiredByPlatform;
 
 	FNativizationSummary() : MemberVariablesFromGraph(0) {}
 };
@@ -92,7 +92,7 @@ public:
 	/**
 	 *	Provides a hook so that external modules can mark some unconverted blueprints as necessary for the generated native code.
 	 */
-	DECLARE_DELEGATE_TwoParams(FMarkUnconvertedBlueprintAsNecessary, TAssetPtr<UBlueprint>, const FCompilerNativizationOptions&);
+	DECLARE_DELEGATE_TwoParams(FMarkUnconvertedBlueprintAsNecessary, TSoftObjectPtr<UBlueprint>, const FCompilerNativizationOptions&);
 	virtual FMarkUnconvertedBlueprintAsNecessary& OnIncludingUnconvertedBP() = 0;
 
 	DECLARE_DELEGATE_RetVal_OneParam(bool, FIsFunctionUsedInADelegate, const UFunction*);
diff --git a/Engine/Source/Developer/BlueprintNativeCodeGen/Private/BlueprintNativeCodeGenManifest.cpp b/Engine/Source/Developer/BlueprintNativeCodeGen/Private/BlueprintNativeCodeGenManifest.cpp
--- a/Engine/Source/Developer/BlueprintNativeCodeGen/Private/BlueprintNativeCodeGenManifest.cpp
+++ b/Engine/Source/Developer/BlueprintNativeCodeGen/Private/BlueprintNativeCodeGenManifest.cpp
@@ -300,12 +300,6 @@ FConvertedAssetRecord::FConvertedAssetRecord(const FAssetData& AssetInfo, const
 	GeneratedHeaderPath = BlueprintNativeCodeGenManifestImpl::GenerateSourceFileSavePath(TargetPaths, AssetInfo, NativizationOptions, FBlueprintNativeCodeGenPaths::HFile);
 }
 
-//------------------------------------------------------------------------------
-FStringAssetReference FConvertedAssetRecord::ToAssetRef() const
-{
-	return FStringAssetReference(TargetObjPath);
-}
-
 /*******************************************************************************
  * FUnconvertedDependencyRecord
  ******************************************************************************/
diff --git a/Engine/Source/Developer/BlueprintNativeCodeGen/Private/BlueprintNativeCodeGenManifest.h b/Engine/Source/Developer/BlueprintNativeCodeGen/Private/BlueprintNativeCodeGenManifest.h
--- a/Engine/Source/Developer/BlueprintNativeCodeGen/Private/BlueprintNativeCodeGenManifest.h
+++ b/Engine/Source/Developer/BlueprintNativeCodeGen/Private/BlueprintNativeCodeGenManifest.h
@@ -4,7 +4,6 @@
 
 #include "CoreMinimal.h"
 #include "UObject/ObjectMacros.h"
-#include "Misc/StringAssetReference.h"
 #include "Engine/Blueprint.h"
 #include "BlueprintNativeCodeGenManifest.generated.h"
 
@@ -24,16 +23,11 @@ public:
 	FConvertedAssetRecord() {}
 	FConvertedAssetRecord(const FAssetData& AssetInfo, const FBlueprintNativeCodeGenPaths& TargetPaths, const FCompilerNativizationOptions& NativizationOptions);
 
-	/**
-	 * @return 
-	 */
-	FStringAssetReference ToAssetRef() const;
-
 public:
 	UPROPERTY()
 	UClass* AssetType;
 
-	// cannot use a FStringAssetReference, as the json serializer has problems 
+	// cannot use a FSoftObjectPath, as the json serializer has problems 
 	// with some asset paths (for example, I had a folder named 'Folder()')
 	UPROPERTY()
 	FString TargetObjPath; 
