Fixed crash caused by retargeting a skeleton when some of its animations are still open

Firstly we now auto-close all assets that use a skeleton when retergetting instead of popping up a dialog to suggest the user does it.
Secondly as a safety measure we dont re-use an animation editor if the persona toolkit's skeleton doesnt match.

diff --git a/Engine/Source/Developer/AssetTools/Private/AssetTypeActions/AssetTypeActions_AnimationAsset.cpp b/Engine/Source/Developer/AssetTools/Private/AssetTypeActions/AssetTypeActions_AnimationAsset.cpp
--- a/Engine/Source/Developer/AssetTools/Private/AssetTypeActions/AssetTypeActions_AnimationAsset.cpp
+++ b/Engine/Source/Developer/AssetTools/Private/AssetTypeActions/AssetTypeActions_AnimationAsset.cpp
@@ -15,6 +15,7 @@
 #include "Framework/MultiBox/MultiBoxBuilder.h"
 #include "Framework/Application/SlateApplication.h"
 #include "Algo/Transform.h"
+#include "IPersonaToolkit.h"
 
 #define LOCTEXT_NAMESPACE "AssetTypeActions"
 
@@ -179,10 +180,13 @@ void FAssetTypeActions_AnimationAsset::OpenAnimAssetEditor(const TArray<UObject*
 							{
 								// Change the current anim to this one
 								IAnimationEditor* AnimEditor = static_cast<IAnimationEditor*>(ExistingEditor);
-								AnimEditor->SetAnimationAsset(AnimAsset);
-								AnimEditor->FocusWindow();
-								bFoundEditor = true;
-								break;
+								if(AnimEditor->GetPersonaToolkit()->GetSkeleton() == AnimAsset->GetSkeleton())
+								{
+									AnimEditor->SetAnimationAsset(AnimAsset);
+									AnimEditor->FocusWindow();
+									bFoundEditor = true;
+									break;
+								}
 							}
 						}
 					}
diff --git a/Engine/Source/Developer/AssetTools/Private/AssetTypeActions/AssetTypeActions_Skeleton.cpp b/Engine/Source/Developer/AssetTools/Private/AssetTypeActions/AssetTypeActions_Skeleton.cpp
--- a/Engine/Source/Developer/AssetTools/Private/AssetTypeActions/AssetTypeActions_Skeleton.cpp
+++ b/Engine/Source/Developer/AssetTools/Private/AssetTypeActions/AssetTypeActions_Skeleton.cpp
@@ -42,6 +42,7 @@
 #include "ISkeletonEditorModule.h"
 #include "Preferences/PersonaOptions.h"
 #include "Algo/Transform.h"
+#include "PhysicsEngine/PhysicsAsset.h"
 
 #define LOCTEXT_NAMESPACE "AssetTypeActions"
 
@@ -595,18 +596,50 @@ void FAssetTypeActions_Skeleton::RetargetAnimationHandler(USkeleton* OldSkeleton
 }
 void FAssetTypeActions_Skeleton::ExecuteRetargetSkeleton(TArray<TWeakObjectPtr<USkeleton>> Skeletons)
 {
-	// only allow 1 for now, it is scary to do this for multi
-	// warn the user to shut down any persona that is opened
-	if ( FMessageDialog::Open(EAppMsgType::YesNo, LOCTEXT("CloseReferencingEditors", "You need to close Persona or anything that references animation, mesh or animation blueprint before this step. Continue?")) == EAppReturnType::Yes )
-	{
-		for (auto SkelIt = Skeletons.CreateConstIterator(); SkelIt; ++SkelIt)
-		{	
-			USkeleton* OldSkeleton = (*SkelIt).Get();
+	TArray<UObject*> AllEditedAssets = FAssetEditorManager::Get().GetAllEditedAssets();
+
+	for (auto SkelIt = Skeletons.CreateConstIterator(); SkelIt; ++SkelIt)
+	{	
+		USkeleton* OldSkeleton = (*SkelIt).Get();
 
-			const FText Message = LOCTEXT("RetargetSkeleton_Warning", "This only converts animation data -i.e. animation assets and Anim Blueprints. \nIf you'd like to convert SkeletalMesh, use the context menu (Assign Skeleton) for each mesh. \n\nIf you'd like to convert mesh as well, please do so before converting animation data. \nOtherwise you will lose any extra track that is in the new mesh.");
-			// ask user what they'd like to change to 
-			SAnimationRemapSkeleton::ShowWindow(OldSkeleton, Message, false, FOnRetargetAnimation::CreateSP(this, &FAssetTypeActions_Skeleton::RetargetAnimationHandler));
+		// Close any assets that reference this skeleton
+		for(UObject* EditedAsset : AllEditedAssets)
+		{
+			bool bCloseAssetEditor = false;
+			if (USkeleton* Skeleton = Cast<USkeleton>(EditedAsset))
+			{
+				bCloseAssetEditor = OldSkeleton == Skeleton;
+			}
+			else if (UAnimationAsset* AnimationAsset = Cast<UAnimationAsset>(EditedAsset))
+			{
+				bCloseAssetEditor = OldSkeleton == AnimationAsset->GetSkeleton();
+			}
+			else if (USkeletalMesh* SkeletalMesh = Cast<USkeletalMesh>(EditedAsset))
+			{
+				bCloseAssetEditor = OldSkeleton == SkeletalMesh->Skeleton;
+			}
+			else if (UAnimBlueprint* AnimBlueprint = Cast<UAnimBlueprint>(EditedAsset))
+			{
+				bCloseAssetEditor = OldSkeleton == AnimBlueprint->TargetSkeleton;
+			}
+			else if (UPhysicsAsset* PhysicsAsset = Cast<UPhysicsAsset>(EditedAsset))
+			{
+				USkeletalMesh* PhysicsSkeletalMesh = PhysicsAsset->PreviewSkeletalMesh.LoadSynchronous();
+				if(PhysicsSkeletalMesh != nullptr)
+				{
+					bCloseAssetEditor = OldSkeleton == PhysicsSkeletalMesh->Skeleton;
+				}
+			}
+
+			if(bCloseAssetEditor)
+			{
+				FAssetEditorManager::Get().CloseAllEditorsForAsset(EditedAsset);
+			}
 		}
+
+		const FText Message = LOCTEXT("RetargetSkeleton_Warning", "This only converts animation data -i.e. animation assets and Anim Blueprints. \nIf you'd like to convert SkeletalMesh, use the context menu (Assign Skeleton) for each mesh. \n\nIf you'd like to convert mesh as well, please do so before converting animation data. \nOtherwise you will lose any extra track that is in the new mesh.");
+		// ask user what they'd like to change to 
+		SAnimationRemapSkeleton::ShowWindow(OldSkeleton, Message, false, FOnRetargetAnimation::CreateSP(this, &FAssetTypeActions_Skeleton::RetargetAnimationHandler));
 	}
 }
 
