Fixed crash if saving over the level you're about to load

diff --git a/Engine/Source/Developer/AssetTools/Private/AssetTypeActions/AssetTypeActions_World.cpp b/Engine/Source/Developer/AssetTools/Private/AssetTypeActions/AssetTypeActions_World.cpp
--- a/Engine/Source/Developer/AssetTools/Private/AssetTypeActions/AssetTypeActions_World.cpp
+++ b/Engine/Source/Developer/AssetTools/Private/AssetTypeActions/AssetTypeActions_World.cpp
@@ -14,16 +14,23 @@ void FAssetTypeActions_World::OpenAssetEditor( const TArray<UObject*>& InObjects
 		UWorld* World = Cast<UWorld>(*ObjIt);
 		if (World != nullptr && ensureMsgf(World->GetTypedOuter<UPackage>(), TEXT("World(%s) is not in a package and cannot be opened"), *World->GetFullName()))
 		{
+			const FString WorldPathName = World->GetPathName();
+
 			// If there are any unsaved changes to the current level, see if the user wants to save those first.
 			bool bPromptUserToSave = true;
 			bool bSaveMapPackages = true;
 			bool bSaveContentPackages = true;
 			if (FEditorFileUtils::SaveDirtyPackages(bPromptUserToSave, bSaveMapPackages, bSaveContentPackages))
 			{
-				const FString FileToOpen = FPackageName::LongPackageNameToFilename(World->GetOutermost()->GetName(), FPackageName::GetMapPackageExtension());
-				const bool bLoadAsTemplate = false;
-				const bool bShowProgress = true;
-				FEditorFileUtils::LoadMap(FileToOpen, bLoadAsTemplate, bShowProgress);
+				// Saving dirty packages may have invalidated the current world pointer (if it was overwritten), so find the world again
+				World = FindObject<UWorld>(nullptr, *WorldPathName);
+				if (World != nullptr && ensureMsgf(World->GetTypedOuter<UPackage>(), TEXT("World(%s) is not in a package and cannot be opened"), *World->GetFullName()))
+				{
+					const FString FileToOpen = FPackageName::LongPackageNameToFilename(World->GetOutermost()->GetName(), FPackageName::GetMapPackageExtension());
+					const bool bLoadAsTemplate = false;
+					const bool bShowProgress = true;
+					FEditorFileUtils::LoadMap(FileToOpen, bLoadAsTemplate, bShowProgress);
+				}
 			}
 			
 			// We can only edit one world at a time... so just break after the first valid world to load
