Moving curve atlas filter to match asset placement in CB
#rnx

diff --git a/Engine/Source/Developer/AssetTools/Private/AssetTypeActions/AssetTypeActions_CurveLinearColorAtlas.h b/Engine/Source/Developer/AssetTools/Private/AssetTypeActions/AssetTypeActions_CurveLinearColorAtlas.h
--- a/Engine/Source/Developer/AssetTools/Private/AssetTypeActions/AssetTypeActions_CurveLinearColorAtlas.h
+++ b/Engine/Source/Developer/AssetTools/Private/AssetTypeActions/AssetTypeActions_CurveLinearColorAtlas.h
@@ -19,4 +19,5 @@ public:
 	virtual void GetActions( const TArray<UObject*>& InObjects, FMenuBuilder& MenuBuilder ) override;
 	virtual bool IsImportedAsset() const override { return false; }
 	virtual bool CanFilter() override { return true; }
+	virtual uint32 GetCategories() override { return EAssetTypeCategories::Misc; }
 };
