Advanced copy now updates references correctly for all Blueprint types. Levels are currently not supported

diff --git a/Engine/Source/Developer/AssetTools/Private/AdvancedCopyCustomization.cpp b/Engine/Source/Developer/AssetTools/Private/AdvancedCopyCustomization.cpp
--- a/Engine/Source/Developer/AssetTools/Private/AdvancedCopyCustomization.cpp
+++ b/Engine/Source/Developer/AssetTools/Private/AdvancedCopyCustomization.cpp
@@ -4,6 +4,9 @@
 #include "Containers/UnrealString.h"
 #include "AssetRegistryModule.h"
 #include "Interfaces/IPluginManager.h"
+#include "Engine/World.h"
+#include "Engine/Level.h"
+#include "Engine/MapBuildDataRegistry.h"
 
 
 #define LOCTEXT_NAMESPACE "AdvancedCopyCustomization"
@@ -24,6 +27,10 @@ UAdvancedCopyCustomization::UAdvancedCopyCustomization(const class FObjectInitia
 
 	FilterForExcludingDependencies.bRecursivePaths = true;
 	FilterForExcludingDependencies.bRecursiveClasses = true;
+	FilterForExcludingDependencies.ClassNames.Add(UWorld::StaticClass()->GetFName());
+	FilterForExcludingDependencies.ClassNames.Add(ULevel::StaticClass()->GetFName());
+	FilterForExcludingDependencies.ClassNames.Add(UMapBuildDataRegistry::StaticClass()->GetFName());
+
 }
 
 void UAdvancedCopyCustomization::SetPackageThatInitiatedCopy(const FString& InBasePackage)
diff --git a/Engine/Source/Developer/AssetTools/Private/AssetTools.cpp b/Engine/Source/Developer/AssetTools/Private/AssetTools.cpp
--- a/Engine/Source/Developer/AssetTools/Private/AssetTools.cpp
+++ b/Engine/Source/Developer/AssetTools/Private/AssetTools.cpp
@@ -806,7 +806,7 @@ bool UAssetToolsImpl::ValidateFlattenedAdvancedCopyDestinations(const TMap<FStri
 	return true;
 }
 
-void UAssetToolsImpl::GetAllAdvancedCopySources(FName SelectedPackage, FAdvancedCopyParams& CopyParams, TArray<FName>& OutPackageNamesToCopy, TMap<FName, FName>& DependencyMap) const
+void UAssetToolsImpl::GetAllAdvancedCopySources(FName SelectedPackage, FAdvancedCopyParams& CopyParams, TArray<FName>& OutPackageNamesToCopy, TMap<FName, FName>& DependencyMap, const UAdvancedCopyCustomization* CopyCustomization) const
 {
 	FAssetRegistryModule& AssetRegistryModule = FModuleManager::Get().LoadModuleChecked<FAssetRegistryModule>(TEXT("AssetRegistry"));
 	IAssetRegistry& AssetRegistry = AssetRegistryModule.Get();
@@ -815,14 +815,23 @@ void UAssetToolsImpl::GetAllAdvancedCopySources(FName SelectedPackage, FAdvanced
 	TArray<FName> CurrentDependencies;
 	if (!OutPackageNamesToCopy.Contains(SelectedPackage))
 	{
-		CurrentDependencies.Add(SelectedPackage);
-		TArray<FAssetData> DependencyAssetData;
-		AssetRegistry.GetAssetsByPackageName(SelectedPackage, DependencyAssetData);
+		TArray<FAssetData> SourceAssetData;
+		AssetRegistry.GetAssetsByPackageName(SelectedPackage, SourceAssetData);
+		// Check if this is a folder before using the filter to exclude assets
+		bool bIsFolder = SourceAssetData.Num() == 0;
+		FARFilter ExclusionFilter = CopyCustomization->GetARFilter();
+		AssetRegistry.UseFilterToExcludeAssets(SourceAssetData, ExclusionFilter);
+		// If this is a valid asset
+		if (SourceAssetData.Num() > 0 || bIsFolder)
+		{
+			CurrentDependencies.Add(SelectedPackage);
+		}
+
 		// If we should check for dependencies OR we are currently checking a folder
 		// Folders should ALWAYS get checked for assets and subfolders
-		if (CopyParams.bShouldCheckForDependencies || DependencyAssetData.Num() == 0)
+		if ((CopyParams.bShouldCheckForDependencies && SourceAssetData.Num() > 0) || bIsFolder)
 		{
-			RecursiveGetDependenciesAdvanced(SelectedPackage, CopyParams, CurrentDependencies, DependencyMap);
+			RecursiveGetDependenciesAdvanced(SelectedPackage, CopyParams, CurrentDependencies, DependencyMap, CopyCustomization);
 		}
 		OutPackageNamesToCopy.Append(CurrentDependencies);
 	}
@@ -832,10 +841,12 @@ bool UAssetToolsImpl::AdvancedCopyPackages(const TMap<FString, FString>& SourceA
 {
 	if (ValidateFlattenedAdvancedCopyDestinations(SourceAndDestPackages))
 	{
-		TArray<FString> SuccessfullyCopiedFiles;
-		TArray<FString> SuccessfullyCopiedPackages;
+		TArray<FString> SuccessfullyCopiedDestinationFiles;
+		TArray<FString> SuccessfullyCopiedSourcePackages;
 		TArray<UObject*> ExistingObjects;
+		TSet<UObject*> ExistingObjectSet;
 		TArray<UObject*> NewObjects;
+		TSet<UObject*> NewObjectSet;
 		FString CopyErrors;
 		FScopedSlowTask LoopProgress(SourceAndDestPackages.Num(), LOCTEXT("AdvancedCopying", "Copying files and dependencies..."));
 		LoopProgress.MakeDialog();
@@ -858,6 +869,7 @@ bool UAssetToolsImpl::AdvancedCopyPackages(const TMap<FString, FString>& SourceA
 					if (ExistingObject)
 					{
 						ExistingObjects.Add(ExistingObject);
+						ExistingObjectSet.Add(ExistingObject);
 						TSet<UPackage*> ObjectsUserRefusedToFullyLoad;
 						ObjectTools::FMoveDialogInfo MoveDialogInfo;
 						MoveDialogInfo.bOkToAll = bCopyOverAllDestinationOverlaps;
@@ -869,39 +881,41 @@ bool UAssetToolsImpl::AdvancedCopyPackages(const TMap<FString, FString>& SourceA
 						const bool bShouldPromptForDestinationConflict = !bCopyOverAllDestinationOverlaps;
 						UObject* NewObject = ObjectTools::DuplicateSingleObject(ExistingObject, MoveDialogInfo.PGN, ObjectsUserRefusedToFullyLoad, bShouldPromptForDestinationConflict);
 						NewObjects.Add(NewObject);
-						SuccessfullyCopiedPackages.Add(PackageName);
-						SuccessfullyCopiedFiles.Add(DestFilename);
+						NewObjectSet.Add(NewObject);
+						SuccessfullyCopiedSourcePackages.Add(PackageName);
+						SuccessfullyCopiedDestinationFiles.Add(DestFilename);
 					}
 				}
 			}
 		}
 
-
 		for (int32 ObjectIdx = 0; ObjectIdx < NewObjects.Num(); ObjectIdx++)
 		{
 			TMap<UObject*, UObject*> ReplacementMap;
 			FAssetRegistryModule& AssetRegistryModule = FModuleManager::Get().LoadModuleChecked<FAssetRegistryModule>(TEXT("AssetRegistry"));
 			TArray<FName> Dependencies;
-			FName SuccessfullyCopiedPackage = FName(*SuccessfullyCopiedPackages[ObjectIdx]);
+			FName SuccessfullyCopiedPackage = FName(*SuccessfullyCopiedSourcePackages[ObjectIdx]);
 			AssetRegistryModule.Get().GetDependencies(SuccessfullyCopiedPackage, Dependencies);
 			for (FName Dependency : Dependencies)
 			{
-				if (SuccessfullyCopiedPackages.Contains(Dependency.ToString()))
+				if (SuccessfullyCopiedSourcePackages.Contains(Dependency.ToString()))
 				{
 					int32 DependencyIndex = ExistingObjects.IndexOfByPredicate([&](UObject* Object) { return Object->GetOuter()->GetName() == Dependency.ToString(); });
 					if (DependencyIndex != INDEX_NONE)
 					{
-						ReplacementMap.Add(ExistingObjects[DependencyIndex], NewObjects[DependencyIndex]);
+						TArray<UObject*> ObjectsToReplace;
+						ObjectsToReplace.Add(ExistingObjects[DependencyIndex]);
+						ObjectTools::ConsolidateObjects(NewObjects[DependencyIndex], ObjectsToReplace, NewObjectSet, ExistingObjectSet, false);
 					}
 				}
 			}
-			FArchiveReplaceObjectRef<UObject> ReplaceAr(NewObjects[ObjectIdx], ReplacementMap, false, true, true);
 		}
 
+		ObjectTools::CompileBlueprintsAfterRefUpdate(NewObjects);
 
 		FString SourceControlErrors;
 
-		if (SuccessfullyCopiedFiles.Num() > 0)
+		if (SuccessfullyCopiedDestinationFiles.Num() > 0)
 		{
 			// attempt to add files to source control (this can quite easily fail, but if it works it is very useful)
 			if (GetDefault<UEditorLoadingSavingSettings>()->bSCCAutoAddNewFiles)
@@ -909,9 +923,9 @@ bool UAssetToolsImpl::AdvancedCopyPackages(const TMap<FString, FString>& SourceA
 				if (ISourceControlModule::Get().IsEnabled())
 				{
 					ISourceControlProvider& SourceControlProvider = ISourceControlModule::Get().GetProvider();
-					if (SourceControlProvider.Execute(ISourceControlOperation::Create<FMarkForAdd>(), SuccessfullyCopiedFiles) == ECommandResult::Failed)
+					if (SourceControlProvider.Execute(ISourceControlOperation::Create<FMarkForAdd>(), SuccessfullyCopiedDestinationFiles) == ECommandResult::Failed)
 					{
-						for (auto FileIt(SuccessfullyCopiedFiles.CreateConstIterator()); FileIt; FileIt++)
+						for (auto FileIt(SuccessfullyCopiedDestinationFiles.CreateConstIterator()); FileIt; FileIt++)
 						{
 							if (!SourceControlProvider.GetState(*FileIt, EStateCacheUsage::Use)->IsAdded())
 							{
@@ -938,13 +952,13 @@ bool UAssetToolsImpl::AdvancedCopyPackages(const TMap<FString, FString>& SourceA
 				ErrorMessage += LINE_TERMINATOR;
 				ErrorMessage += LOCTEXT("AdvancedCopyPackages_SourceControlErrorsList", "Some files reported source control errors.").ToString();
 			}
-			if (SuccessfullyCopiedPackages.Num() > 0)
+			if (SuccessfullyCopiedSourcePackages.Num() > 0)
 			{
 				AdvancedCopyLog.NewPage(LOCTEXT("AdvancedCopyPackages_CopyErrorsSuccesslistPage", "Copied Successfully"));
 				AdvancedCopyLog.Info(FText::FromString(*SourceControlErrors));
 				ErrorMessage += LINE_TERMINATOR;
 				ErrorMessage += LOCTEXT("AdvancedCopyPackages_CopyErrorsSuccesslist", "Some files were copied successfully.").ToString();
-				for (auto FileIt = SuccessfullyCopiedPackages.CreateConstIterator(); FileIt; ++FileIt)
+				for (auto FileIt = SuccessfullyCopiedSourcePackages.CreateConstIterator(); FileIt; ++FileIt)
 				{
 					if (FileIt->Len() > 0)
 					{
@@ -957,7 +971,7 @@ bool UAssetToolsImpl::AdvancedCopyPackages(const TMap<FString, FString>& SourceA
 		else
 		{
 			AdvancedCopyLog.NewPage(LOCTEXT("AdvancedCopyPackages_CompletePage", "Advanced content copy completed successfully!"));
-			for (auto FileIt = SuccessfullyCopiedPackages.CreateConstIterator(); FileIt; ++FileIt)
+			for (auto FileIt = SuccessfullyCopiedSourcePackages.CreateConstIterator(); FileIt; ++FileIt)
 			{
 				if (FileIt->Len() > 0)
 				{
@@ -2835,7 +2849,7 @@ void UAssetToolsImpl::RecursiveGetDependencies(const FName& PackageName, TSet<FN
 	}
 }
 
-void UAssetToolsImpl::RecursiveGetDependenciesAdvanced(const FName& PackageName, FAdvancedCopyParams& CopyParams, TArray<FName>& AllDependencies, TMap<FName, FName>& DependencyMap) const
+void UAssetToolsImpl::RecursiveGetDependenciesAdvanced(const FName& PackageName, FAdvancedCopyParams& CopyParams, TArray<FName>& AllDependencies, TMap<FName, FName>& DependencyMap, const UAdvancedCopyCustomization* CopyCustomization) const
 {
 	FAssetRegistryModule& AssetRegistryModule = FModuleManager::Get().LoadModuleChecked<FAssetRegistryModule>(TEXT("AssetRegistry"));
 	TArray<FName> Dependencies;
@@ -2853,13 +2867,13 @@ void UAssetToolsImpl::RecursiveGetDependenciesAdvanced(const FName& PackageName,
 				TArray<FAssetData> DependencyAssetData;
 				if (AssetRegistry.GetAssetsByPackageName(*DependsIt, DependencyAssetData))
 				{
-					FARFilter ExclusionFilter = UAdvancedCopyCustomization::StaticClass()->GetDefaultObject<UAdvancedCopyCustomization>()->GetARFilter();
+					FARFilter ExclusionFilter = CopyCustomization->GetARFilter();
 					AssetRegistry.UseFilterToExcludeAssets(DependencyAssetData, ExclusionFilter);
 					if (DependencyAssetData.IsValidIndex(0))
 					{
 						AllDependencies.Add(*DependsIt);
 						DependencyMap.Add(*DependsIt, PackageName);
-						RecursiveGetDependenciesAdvanced(*DependsIt, CopyParams, AllDependencies, DependencyMap);
+						RecursiveGetDependenciesAdvanced(*DependsIt, CopyParams, AllDependencies, DependencyMap, CopyCustomization);
 					}
 				}
 
@@ -2880,7 +2894,7 @@ void UAssetToolsImpl::RecursiveGetDependenciesAdvanced(const FName& PackageName,
 				// If we should check the assets we found for dependencies
 				if (CopyParams.bShouldCheckForDependencies)
 				{
-					RecursiveGetDependenciesAdvanced(FName(*Asset.GetPackage()->GetName()), CopyParams, AllDependencies, DependencyMap);
+					RecursiveGetDependenciesAdvanced(FName(*Asset.GetPackage()->GetName()), CopyParams, AllDependencies, DependencyMap, CopyCustomization);
 				}
 			}
 		}
@@ -2891,7 +2905,7 @@ void UAssetToolsImpl::RecursiveGetDependenciesAdvanced(const FName& PackageName,
 			AssetRegistry.GetSubPaths(PackageName.ToString(), SubPaths, false);
 			for (const FString SubPath : SubPaths)
 			{
-				RecursiveGetDependenciesAdvanced(FName(*SubPath), CopyParams, AllDependencies, DependencyMap);
+				RecursiveGetDependenciesAdvanced(FName(*SubPath), CopyParams, AllDependencies, DependencyMap, CopyCustomization);
 			}
 
 		}
@@ -3010,7 +3024,7 @@ void UAssetToolsImpl::InitAdvancedCopyFromCopyParams(FAdvancedCopyParams CopyPar
 		TArray<FName> PackageNamesToCopy;
 
 		// Get all packages to be copied
-		GetAllAdvancedCopySources(Package, CopyParams, PackageNamesToCopy, DependencyMap);
+		GetAllAdvancedCopySources(Package, CopyParams, PackageNamesToCopy, DependencyMap, CopyCustomization);
 		// Allow the customization to apply any additional filters
 		CopyCustomization->ApplyAdditionalFiltering(PackageNamesToCopy);
 
diff --git a/Engine/Source/Developer/AssetTools/Private/AssetTools.h b/Engine/Source/Developer/AssetTools/Private/AssetTools.h
--- a/Engine/Source/Developer/AssetTools/Private/AssetTools.h
+++ b/Engine/Source/Developer/AssetTools/Private/AssetTools.h
@@ -105,7 +105,7 @@ public:
 	virtual void GenerateAdvancedCopyDestinations(FAdvancedCopyParams& InParams, const TArray<FName>& InPackageNamesToCopy, const class UAdvancedCopyCustomization* CopyCustomization, TMap<FString, FString>& OutPackagesAndDestinations) const override;
 	virtual bool FlattenAdvancedCopyDestinations(const TArray<TMap<FString, FString>> PackagesAndDestinations, TMap<FString, FString>& FlattenedPackagesAndDestinations) const override;
 	virtual bool ValidateFlattenedAdvancedCopyDestinations(const TMap<FString, FString>& FlattenedPackagesAndDestinations) const override;
-	virtual void GetAllAdvancedCopySources(FName SelectedPackage, FAdvancedCopyParams& CopyParams, TArray<FName>& OutPackageNamesToCopy, TMap<FName, FName>& DependencyMap) const override;
+	virtual void GetAllAdvancedCopySources(FName SelectedPackage, FAdvancedCopyParams& CopyParams, TArray<FName>& OutPackageNamesToCopy, TMap<FName, FName>& DependencyMap, const class UAdvancedCopyCustomization* CopyCustomization) const override;
 	virtual void InitAdvancedCopyFromCopyParams(FAdvancedCopyParams CopyParams) const override;
 
 	virtual void OpenEditorForAssets(const TArray<UObject*>& Assets) const override;
@@ -145,7 +145,7 @@ private:
 	void RecursiveGetDependencies(const FName& PackageName, TSet<FName>& AllDependencies) const;
 
 	/** Gets the dependencies of the specified package recursively while omitting things that don't pass the FARFilter passed in from FAdvancedCopyParams */
-	void RecursiveGetDependenciesAdvanced(const FName& PackageName, FAdvancedCopyParams& CopyParams, TArray<FName>& AllDependencies, TMap<FName, FName>& DependencyMap) const;
+	void RecursiveGetDependenciesAdvanced(const FName& PackageName, FAdvancedCopyParams& CopyParams, TArray<FName>& AllDependencies, TMap<FName, FName>& DependencyMap, const class UAdvancedCopyCustomization* CopyCustomization) const;
 
 	/** Records the time taken for an import and reports it to engine analytics, if available */
 	static void OnNewImportRecord(UClass* AssetType, const FString& FileExtension, bool bSucceeded, bool bWasCancelled, const FDateTime& StartTime);
diff --git a/Engine/Source/Developer/AssetTools/Private/SAdvancedCopyReportDialog.cpp b/Engine/Source/Developer/AssetTools/Private/SAdvancedCopyReportDialog.cpp
--- a/Engine/Source/Developer/AssetTools/Private/SAdvancedCopyReportDialog.cpp
+++ b/Engine/Source/Developer/AssetTools/Private/SAdvancedCopyReportDialog.cpp
@@ -140,7 +140,7 @@ void SAdvancedCopyReportDialog::Construct( const FArguments& InArgs, const FAdva
 			.Padding(0, 4)
 			[
 				SNew(STextBlock)
-				.Text(FText::Format(LOCTEXT("AdvancedCopyDesc","The following files will be copied to {0} and references to copied files will be fixed up."),InReportMessage))
+				.Text(this, &SAdvancedCopyReportDialog::GetHeaderText, InReportMessage)
 				.TextStyle( FEditorStyle::Get(), "PackageMigration.DialogTitle" )
 			]
 
@@ -222,6 +222,15 @@ void SAdvancedCopyReportDialog::Construct( const FArguments& InArgs, const FAdva
 }
 END_SLATE_FUNCTION_BUILD_OPTIMIZATION
 
+FText SAdvancedCopyReportDialog::GetHeaderText(const FText InReportMessage) const
+{
+	if (PackageReportRootNode.Children.Num() == 0)
+	{
+		return LOCTEXT("NoValidSources", "You have not selected any valid sources for advanced copying.");
+	}
+	return FText::Format(LOCTEXT("AdvancedCopyDesc", "The following files will be copied to {0} and references to copied files will be fixed up."), InReportMessage);
+}
+
 ECheckBoxState SAdvancedCopyReportDialog::IsGeneratingDependencies() const
 {
 	return CurrentCopyParams.bShouldCheckForDependencies ? ECheckBoxState::Checked : ECheckBoxState::Unchecked;
diff --git a/Engine/Source/Developer/AssetTools/Private/SAdvancedCopyReportDialog.h b/Engine/Source/Developer/AssetTools/Private/SAdvancedCopyReportDialog.h
--- a/Engine/Source/Developer/AssetTools/Private/SAdvancedCopyReportDialog.h
+++ b/Engine/Source/Developer/AssetTools/Private/SAdvancedCopyReportDialog.h
@@ -68,7 +68,7 @@ public:
 
 	/** Constructs this widget with InArgs */
 	void Construct( const FArguments& InArgs, const FAdvancedCopyParams& InParams, const FText& InReportMessage, const TArray<TMap<FString, FString>>& DestinationMap, const TArray<TMap<FName, FName>>& DependencyMap, const FOnReportConfirmed& InOnReportConfirmed );
-
+		
 	ECheckBoxState IsGeneratingDependencies() const;
 
 	void ToggleGeneratingDependencies(ECheckBoxState NewState);
@@ -99,6 +99,8 @@ private:
 	/** Handler for when "Cancel" is clicked */
 	FReply CancelClicked();
 
+	FText GetHeaderText(const FText InReportMessage) const;
+
 private:
 	FOnReportConfirmed OnReportConfirmed;
 	FAdvancedCopyReportNode PackageReportRootNode;
diff --git a/Engine/Source/Developer/AssetTools/Public/IAssetTools.h b/Engine/Source/Developer/AssetTools/Public/IAssetTools.h
--- a/Engine/Source/Developer/AssetTools/Public/IAssetTools.h
+++ b/Engine/Source/Developer/AssetTools/Public/IAssetTools.h
@@ -392,7 +392,7 @@ public:
 	virtual bool ValidateFlattenedAdvancedCopyDestinations(const TMap<FString, FString>& FlattenedPackagesAndDestinations) const = 0;
 
 	/* Find all the dependencies that also need to be copied in the advanced copy, mapping them to the file that depends on them and excluding any that don't pass the ARFilter stored on CopyParams */
-	virtual void GetAllAdvancedCopySources(FName SelectedPackage, FAdvancedCopyParams& CopyParams, TArray<FName>& OutPackageNamesToCopy, TMap<FName, FName>& DependencyMap) const = 0;
+	virtual void GetAllAdvancedCopySources(FName SelectedPackage, FAdvancedCopyParams& CopyParams, TArray<FName>& OutPackageNamesToCopy, TMap<FName, FName>& DependencyMap, const class UAdvancedCopyCustomization* CopyCustomization) const = 0;
 
 	/* Given a complete set of copy parameters, which includes the selected package set, start the advanced copy process */
 	virtual void InitAdvancedCopyFromCopyParams(FAdvancedCopyParams CopyParams) const = 0;
diff --git a/Engine/Source/Editor/UnrealEd/Private/ObjectTools.cpp b/Engine/Source/Editor/UnrealEd/Private/ObjectTools.cpp
--- a/Engine/Source/Editor/UnrealEd/Private/ObjectTools.cpp
+++ b/Engine/Source/Editor/UnrealEd/Private/ObjectTools.cpp
@@ -91,6 +91,7 @@
 #include "Templates/UniquePtr.h"
 #include "Engine/MapBuildDataRegistry.h"
 #include "HAL/PlatformApplicationMisc.h"
+#include "Kismet2/BlueprintEditorUtils.h"
 
 DEFINE_LOG_CATEGORY_STATIC(LogObjectTools, Log, All);
 
@@ -622,11 +623,10 @@ namespace ObjectTools
 	 * @param bWarnAboutRootSet		If True a message will be displayed to a user asking them if they would like to remove the rootset flag from objects which have it set.  
 									If False, the message will not be displayed and rootset is automatically removed 
 	 */
-	static void ForceReplaceReferences( UObject* ObjectToReplaceWith, TArray<UObject*>& ObjectsToReplace, FForceReplaceInfo& OutInfo, bool bWarnAboutRootSet = true)				
+	void ForceReplaceReferences( UObject* ObjectToReplaceWith, TArray<UObject*>& ObjectsToReplace, TSet<UObject*>& ObjectsToReplaceWithin, FForceReplaceInfo& OutInfo, bool bWarnAboutRootSet = true)
 	{
 		FPropertyEditorModule& PropertyEditorModule = FModuleManager::LoadModuleChecked<FPropertyEditorModule>("PropertyEditor");
 		PropertyEditorModule.RemoveDeletedObjects( ObjectsToReplace );
-
 		TSet<UObject*> RootSetObjects;
 
 		GWarn->StatusUpdate( 0, 0, NSLOCTEXT("UnrealEd", "ConsolidateAssetsUpdate_RootSetCheck", "Checking Assets for Root Set...") );
@@ -662,10 +662,10 @@ namespace ObjectTools
 				FText Message = FText::Format( MessageFormatting, Arguments );
 
 				// Prompt the user to see if they'd like to remove the root set flag from the assets and attempt to replace them
-				EAppReturnType::Type UserRepsonse = OpenMsgDlgInt( EAppMsgType::YesNo, EAppReturnType::No, Message, NSLOCTEXT("ObjectTools", "ConsolidateAssetsRootSetDlg_Title", "Failed to Consolidate Assets") );
+				EAppReturnType::Type UserResponse = OpenMsgDlgInt( EAppMsgType::YesNo, EAppReturnType::No, Message, NSLOCTEXT("ObjectTools", "ConsolidateAssetsRootSetDlg_Title", "Failed to Consolidate Assets") );
 
 				// The user elected to not remove the root set flag, so cancel the replacement
-				if ( UserRepsonse == EAppReturnType::No )
+				if (UserResponse == EAppReturnType::No )
 				{
 					return;
 				}
@@ -673,6 +673,7 @@ namespace ObjectTools
 
 			for ( FObjectIterator ObjIter; ObjIter; ++ObjIter )
 			{
+				// Always clear the root set flags
 				UObject* CurrentObject = *ObjIter;
 				if ( CurrentObject )
 				{
@@ -765,8 +766,13 @@ namespace ObjectTools
 		{
 			UObject* CurObject = *ObjIter;
 
+			if (ObjectsToReplaceWithin.Num() > 0 && !ObjectsToReplaceWithin.Contains(CurObject))
+			{
+				CurObject = nullptr;
+			}
+
 			// Unless the "object to replace with" is null, ignore any of the objects to replace to themselves
-			if ( ObjectToReplaceWith == NULL || !ReplacementMap.Find( CurObject ) )
+			if (CurObject && ( ObjectToReplaceWith == NULL || !ReplacementMap.Find( CurObject ) ))
 			{
 				// Find the referencers of the objects to be replaced
 				FFindReferencersArchive FindRefsArchive( CurObject, OutInfo.ReplaceableObjects );
@@ -824,6 +830,20 @@ namespace ObjectTools
 			}
 		}
 
+		if(ObjectsToReplaceWithin.Num() > 0)
+		{
+			for (UObject* CurObject : ObjectsToReplaceWithin)
+			{
+				UBlueprint* BPObjectToUpdate = Cast<UBlueprint>(CurObject);
+				if (BPObjectToUpdate)
+				{
+					FArchiveReplaceObjectRef<UObject> ReplaceAr(BPObjectToUpdate->GeneratedClass->ClassDefaultObject, ReplacementMap, false, true, false);
+				}
+				FArchiveReplaceObjectRef<UObject> ReplaceAr(CurObject, ReplacementMap, false, true, false);
+			}
+		}
+		else
+		{
 		// Iterate over the map of referencing objects/changed properties, forcefully replacing the references and
 		int32 NumObjsReplaced = 0;
 		for (int32 Index = 0; Index < ReferencingPropertiesMapKeys.Num(); Index++)
@@ -835,7 +855,7 @@ namespace ObjectTools
 
 			FArchiveReplaceObjectRef<UObject> ReplaceAr( CurReplaceObj, ReplacementMap, false, true, false );
 		}
-
+		}
 		// Now alter the referencing objects the change has completed via PostEditChange, 
 		// this is done in a separate loop to prevent reading of data that we want to overwrite
 		int32 NumObjsPostEdited = 0;
@@ -872,22 +892,28 @@ namespace ObjectTools
 		}
 	}
 
-	FConsolidationResults ConsolidateObjects( UObject* ObjectToConsolidateTo, TArray<UObject*>& ObjectsToConsolidate, bool bShowDeleteConfirmation )
+	/**
+	 * Forcefully replaces references to passed in objects
+	 *
+	 * @param ObjectToReplaceWith	Any references found to 'ObjectsToReplace' will be replaced with this object.  If the object is NULL references will be nulled.
+	 * @param ObjectsToReplace		An array of objects that should be replaced with 'ObjectToReplaceWith'
+	 * @param OutInfo				FForceReplaceInfo struct containing useful information about the result of the call to this function
+	 * @param bWarnAboutRootSet		If True a message will be displayed to a user asking them if they would like to remove the rootset flag from objects which have it set.
+									If False, the message will not be displayed and rootset is automatically removed
+	 */
+	static void ForceReplaceReferences(UObject* ObjectToReplaceWith, TArray<UObject*>& ObjectsToReplace, FForceReplaceInfo& OutInfo, bool bWarnAboutRootSet = true)
+	{
+		TSet<UObject*> InObjectsToReplaceWithin;
+		ForceReplaceReferences(ObjectToReplaceWith, ObjectsToReplace, InObjectsToReplaceWithin, OutInfo, bWarnAboutRootSet);
+	}
+
+	FConsolidationResults ConsolidateObjects( UObject* ObjectToConsolidateTo, TArray<UObject*>& ObjectsToConsolidate, TSet<UObject*>& ObjectsToConsolidateWithin, TSet<UObject*>& ObjectsToNotConsolidateWithin, bool bShouldDeleteAfterConsolidate)
 	{
 		FConsolidationResults ConsolidationResults;
 
 		// Ensure the consolidation is headed toward a valid object and this isn't occurring in game
 		if ( ObjectToConsolidateTo )
 		{
-			// Confirm that the consolidate was intentional
-			if ( bShowDeleteConfirmation )
-			{
-				if ( !ShowDeleteConfirmationDialog( ObjectsToConsolidate ) )
-				{
-					return ConsolidationResults;
-				}
-			}
-
 			// Close all editors to avoid changing references to temporary objects used by the editor
 			if ( !FAssetEditorManager::Get().CloseAllAssetEditors() )
 			{
@@ -1044,7 +1070,7 @@ namespace ObjectTools
 					}
 				}
 
-				ForceReplaceReferences(ObjectToConsolidateTo, ObjectsToConsolidate, ReplaceInfo);
+				ForceReplaceReferences(ObjectToConsolidateTo, ObjectsToConsolidate, ObjectsToConsolidateWithin, ReplaceInfo);
 
 				if (UBlueprint* ObjectToConsolidateTo_BP = Cast<UBlueprint>(ObjectToConsolidateTo))
 				{
@@ -1068,7 +1094,7 @@ namespace ObjectTools
 						}
 					}
 
-					ForceReplaceReferences(ObjectToConsolidateTo_BP->GeneratedClass, ObjectsToConsolidate_BP, GeneratedClassReplaceInfo);
+					ForceReplaceReferences(ObjectToConsolidateTo_BP->GeneratedClass, ObjectsToConsolidate_BP, ObjectsToConsolidateWithin, GeneratedClassReplaceInfo);
 
 					// Repair the references of GeneratedClass on the object being consolidated so they can be properly disposed of upon deletion.
 					for (int32 Index = 0, MaxIndex = ObjectsToConsolidate.Num(); Index < MaxIndex; ++Index)
@@ -1101,7 +1127,9 @@ namespace ObjectTools
 						// then repair the references on the object being consolidated so those objects can be properly disposed of upon deletion.
 						UClass* OldClass = BlueprintToConsolidate->GeneratedClass;
 						UClass* OldSkeletonClass = BlueprintToConsolidate->SkeletonGeneratedClass;
-						FBlueprintCompileReinstancer::ReplaceInstancesOfClass(OldClass, BlueprintToConsolidateTo->GeneratedClass, nullptr, nullptr, true );
+
+
+						FBlueprintCompileReinstancer::ReplaceInstancesOfClass(OldClass, BlueprintToConsolidateTo->GeneratedClass, nullptr, &ObjectsToNotConsolidateWithin, true);
 						BlueprintToConsolidate->GeneratedClass = OldClass;
 						BlueprintToConsolidate->SkeletonGeneratedClass = OldSkeletonClass;
 					}
@@ -1119,6 +1147,8 @@ namespace ObjectTools
 
 			TSet<FString> AlreadyMappedObjectPaths;
 
+			if (bShouldDeleteAfterConsolidate)
+			{
 			// With all references to the objects to consolidate to eliminated from objects that are currently loaded, it should now be safe to delete
 			// the objects to be consolidated themselves, leaving behind a redirector in their place to fix up objects that were not currently loaded at the time
 			// of this operation.
@@ -1191,10 +1221,6 @@ namespace ObjectTools
 
 			CleanupAfterSuccessfulDelete(PotentialPackagesToDelete);
 
-			// Empty the provided array so it's not full of pointers to deleted objects
-			ObjectsToConsolidate.Empty();
-			ConsolidatedObjects.Empty();
-
 			// Now that the old objects have been garbage collected, give the redirectors a proper name
 			for (TMap<UObjectRedirector*, FName>::TIterator RedirectIt(RedirectorToObjectNameMap); RedirectIt; ++RedirectIt)
 			{
@@ -1213,6 +1239,11 @@ namespace ObjectTools
 					CriticalFailureObjects.AddUnique(Redirector);
 				}
 			}
+			}
+
+			// Empty the provided array so it's not full of pointers to deleted objects
+			ObjectsToConsolidate.Empty();
+			ConsolidatedObjects.Empty();
 
 			GWarn->EndSlowTask();
 
@@ -1268,6 +1299,44 @@ namespace ObjectTools
 		return ConsolidationResults;
 	}
 
+	FConsolidationResults ConsolidateObjects(UObject* ObjectToConsolidateTo, TArray<UObject*>& ObjectsToConsolidate, bool bShowDeleteConfirmation)
+	{
+		FConsolidationResults ConsolidationResults;
+
+		// Ensure the consolidation is headed toward a valid object and this isn't occurring in game
+		if (ObjectToConsolidateTo)
+		{
+			// Confirm that the consolidate was intentional
+			if (bShowDeleteConfirmation)
+			{
+				if (!ShowDeleteConfirmationDialog(ObjectsToConsolidate))
+				{
+					return ConsolidationResults;
+				}
+			}
+			TSet<UObject*> ObjectsToConsolidateWithin;
+			TSet<UObject*> ObjectsToNotConsolidateWithin;
+			return ConsolidateObjects(ObjectToConsolidateTo, ObjectsToConsolidate, ObjectsToConsolidateWithin, ObjectsToNotConsolidateWithin, true);
+		}
+
+		return ConsolidationResults;
+	}
+
+	void CompileBlueprintsAfterRefUpdate(TArray<UObject*>& ObjectsConsolidatedWithin)
+	{
+		for (UObject* CurObject : ObjectsConsolidatedWithin)
+		{
+			if (CurObject)
+			{
+				UBlueprint* BPObjectToUpdate = Cast<UBlueprint>(CurObject);
+				if (BPObjectToUpdate)
+				{
+					FKismetEditorUtilities::CompileBlueprint(BPObjectToUpdate);
+				}
+			}
+		}
+	}
+
 	/**
 	 * Copies references for selected generic browser objects to the clipboard.
 	 */
diff --git a/Engine/Source/Editor/UnrealEd/Public/ObjectTools.h b/Engine/Source/Editor/UnrealEd/Public/ObjectTools.h
--- a/Engine/Source/Editor/UnrealEd/Public/ObjectTools.h
+++ b/Engine/Source/Editor/UnrealEd/Public/ObjectTools.h
@@ -194,7 +194,8 @@ namespace ObjectTools
 	 * @return	Structure of consolidation results, specifying which packages were dirtied, which objects failed consolidation (if any), etc.
 	 */
 	UNREALED_API FConsolidationResults ConsolidateObjects( UObject* ObjectToConsolidateTo, TArray<UObject*>& ObjectsToConsolidate, bool bShowDeleteConfirmation = true );
-
+	UNREALED_API FConsolidationResults ConsolidateObjects(UObject* ObjectToConsolidateTo, TArray<UObject*>& ObjectsToConsolidate, TSet<UObject*>& ObjectsToConsolidateWithin, TSet<UObject*>& ObjectsToNotConsolidateWithin, bool bShouldDeleteAfterConsolidate);
+	UNREALED_API void CompileBlueprintsAfterRefUpdate(TArray<UObject*>& ObjectsConsolidatedWithin);
 	/**
 	 * Copies references for selected generic browser objects to the clipboard.
 	 */
