Rename Blueprint Foliage to Actor Foliage as discussed.

#rnx

diff --git a/Engine/Source/Developer/AssetTools/Private/AssetTypeActions/AssetTypeActions_ActorFoliageSettings.h b/Engine/Source/Developer/AssetTools/Private/AssetTypeActions/AssetTypeActions_ActorFoliageSettings.h
--- a/Engine/Source/Developer/AssetTools/Private/AssetTypeActions/AssetTypeActions_ActorFoliageSettings.h
+++ b/Engine/Source/Developer/AssetTools/Private/AssetTypeActions/AssetTypeActions_ActorFoliageSettings.h
@@ -13,7 +13,7 @@ public:
 	{}
 
 	// IAssetTypeActions Implementation
-	virtual FText GetName() const override { return NSLOCTEXT("AssetTypeActions", "AssetTypeActions_ActorFoliageSettings", "Blueprint Foliage"); }
+	virtual FText GetName() const override { return NSLOCTEXT("AssetTypeActions", "AssetTypeActions_ActorFoliageSettings", "Actor Foliage"); }
 	virtual FColor GetTypeColor() const override { return FColor(12, 65, 12); }
 	virtual UClass* GetSupportedClass() const override;
 	virtual uint32 GetCategories() override { return AssetCategoryBit; }
diff --git a/Engine/Source/Editor/FoliageEdit/Private/FoliageTypeFactory.cpp b/Engine/Source/Editor/FoliageEdit/Private/FoliageTypeFactory.cpp
--- a/Engine/Source/Editor/FoliageEdit/Private/FoliageTypeFactory.cpp
+++ b/Engine/Source/Editor/FoliageEdit/Private/FoliageTypeFactory.cpp
@@ -40,7 +40,7 @@ UObject* UFoliageType_ActorFactory::FactoryCreateNew(UClass* Class, UObject* InP
 
 FText UFoliageType_ActorFactory::GetToolTip() const
 {
-	return LOCTEXT("FoliageTypeActorToolTip", "Blueprint Foliage is a foliage type that will place blueprint actor instances. The cost of painting this foliage is the same as adding actors in a scene.");
+	return LOCTEXT("FoliageTypeActorToolTip", "Actor Foliage is a foliage type that will place blueprint/native actor instances. The cost of painting this foliage is the same as adding actors in a scene.");
 }
 
 #undef LOCTEXT_NAMESPACE
