Disable creating streamable animations in 4.23. Feature will be reworked for 4.24

diff --git a/Engine/Source/Developer/AssetTools/Private/AssetTypeActions/AssetTypeActions_AnimSequence.cpp b/Engine/Source/Developer/AssetTools/Private/AssetTypeActions/AssetTypeActions_AnimSequence.cpp
--- a/Engine/Source/Developer/AssetTools/Private/AssetTypeActions/AssetTypeActions_AnimSequence.cpp
+++ b/Engine/Source/Developer/AssetTools/Private/AssetTypeActions/AssetTypeActions_AnimSequence.cpp
@@ -79,6 +79,7 @@ void FAssetTypeActions_AnimSequence::FillCreateMenu(FMenuBuilder& MenuBuilder, c
 			)
 		);
 
+	/* Not supported, streamable animation logic will be ported to UAnimSequence
 	MenuBuilder.AddMenuEntry(
 		LOCTEXT("AnimSequence_NewAnimStreamable", "Create AnimStreamable"),
 		LOCTEXT("AnimSequence_NewAnimStreamableTooltip", "Creates an AnimStreamable using the selected anim sequence."),
@@ -87,7 +88,7 @@ void FAssetTypeActions_AnimSequence::FillCreateMenu(FMenuBuilder& MenuBuilder, c
 			FExecuteAction::CreateSP(this, &FAssetTypeActions_AnimSequence::ExecuteNewAnimStreamable, Sequences),
 			FCanExecuteAction()
 		)
-	);
+	);*/
 
 	MenuBuilder.AddMenuEntry(
 		LOCTEXT("AnimSequence_NewPoseAsset", "Create PoseAsset"),
