Guard against null factories causing a crash when creating a new asset

diff --git a/Engine/Source/Developer/AssetTools/Private/AssetTools.cpp b/Engine/Source/Developer/AssetTools/Private/AssetTools.cpp
--- a/Engine/Source/Developer/AssetTools/Private/AssetTools.cpp
+++ b/Engine/Source/Developer/AssetTools/Private/AssetTools.cpp
@@ -523,25 +523,28 @@ UObject* UAssetToolsImpl::CreateAssetWithDialog(UClass* AssetClass, UFactory* Fa
 
 UObject* UAssetToolsImpl::CreateAssetWithDialog(const FString& AssetName, const FString& PackagePath, UClass* AssetClass, UFactory* Factory, FName CallingContext)
 {
-	FSaveAssetDialogConfig SaveAssetDialogConfig;
-	SaveAssetDialogConfig.DialogTitleOverride = LOCTEXT("SaveAssetDialogTitle", "Save Asset As");
-	SaveAssetDialogConfig.DefaultPath = PackagePath;
-	SaveAssetDialogConfig.DefaultAssetName = AssetName;
-	SaveAssetDialogConfig.ExistingAssetPolicy = ESaveAssetDialogExistingAssetPolicy::AllowButWarn;
-
-	FContentBrowserModule& ContentBrowserModule = FModuleManager::LoadModuleChecked<FContentBrowserModule>("ContentBrowser");
-	FString SaveObjectPath = ContentBrowserModule.Get().CreateModalSaveAssetDialog(SaveAssetDialogConfig);
-	if (!SaveObjectPath.IsEmpty())
+	if(Factory)
 	{
-		FEditorDelegates::OnConfigureNewAssetProperties.Broadcast(Factory);
-		if (Factory->ConfigureProperties())
+		FSaveAssetDialogConfig SaveAssetDialogConfig;
+		SaveAssetDialogConfig.DialogTitleOverride = LOCTEXT("SaveAssetDialogTitle", "Save Asset As");
+		SaveAssetDialogConfig.DefaultPath = PackagePath;
+		SaveAssetDialogConfig.DefaultAssetName = AssetName;
+		SaveAssetDialogConfig.ExistingAssetPolicy = ESaveAssetDialogExistingAssetPolicy::AllowButWarn;
+
+		FContentBrowserModule& ContentBrowserModule = FModuleManager::LoadModuleChecked<FContentBrowserModule>("ContentBrowser");
+		FString SaveObjectPath = ContentBrowserModule.Get().CreateModalSaveAssetDialog(SaveAssetDialogConfig);
+		if (!SaveObjectPath.IsEmpty())
 		{
-			const FString SavePackageName = FPackageName::ObjectPathToPackageName(SaveObjectPath);
-			const FString SavePackagePath = FPaths::GetPath(SavePackageName);
-			const FString SaveAssetName = FPaths::GetBaseFilename(SavePackageName);
-			FEditorDirectories::Get().SetLastDirectory(ELastDirectory::NEW_ASSET, PackagePath);
+			FEditorDelegates::OnConfigureNewAssetProperties.Broadcast(Factory);
+			if (Factory->ConfigureProperties())
+			{
+				const FString SavePackageName = FPackageName::ObjectPathToPackageName(SaveObjectPath);
+				const FString SavePackagePath = FPaths::GetPath(SavePackageName);
+				const FString SaveAssetName = FPaths::GetBaseFilename(SavePackageName);
+				FEditorDirectories::Get().SetLastDirectory(ELastDirectory::NEW_ASSET, PackagePath);
 
-			return CreateAsset(SaveAssetName, SavePackagePath, AssetClass, Factory, CallingContext);
+				return CreateAsset(SaveAssetName, SavePackagePath, AssetClass, Factory, CallingContext);
+			}
 		}
 	}
 
