When renaming a map, leave map built data inside the map's package

diff --git a/Engine/Source/Developer/AssetTools/Private/AssetRenameManager.cpp b/Engine/Source/Developer/AssetTools/Private/AssetRenameManager.cpp
--- a/Engine/Source/Developer/AssetTools/Private/AssetRenameManager.cpp
+++ b/Engine/Source/Developer/AssetTools/Private/AssetRenameManager.cpp
@@ -260,10 +260,10 @@ bool FAssetRenameManager::FixReferencesAndRename(const TArray<FAssetRenameData>&
 	TArray<FAssetRenameDataWithReferencers> AssetsToRename;
 	AssetsToRename.Reset(AssetsAndNames.Num());
 	// Avoid duplicates when adding MapBuildData to list
-	TMap<UObject*, bool> AssetsToRenameLookup;
+	TSet<UObject*> AssetsToRenameLookup;
 	for (const FAssetRenameData& AssetRenameData : AssetsAndNames)
 	{
-		AssetsToRenameLookup.FindOrAdd(AssetRenameData.Asset.Get());
+		AssetsToRenameLookup.Add(AssetRenameData.Asset.Get());
 	}
 	for (const FAssetRenameData& AssetRenameData : AssetsAndNames)
 	{
@@ -273,11 +273,15 @@ bool FAssetRenameManager::FixReferencesAndRename(const TArray<FAssetRenameData>&
 			UWorld* World = Cast<UWorld>(AssetRenameData.Asset.Get());
 			if (World && World->PersistentLevel && World->PersistentLevel->MapBuildData && !AssetsToRenameLookup.Contains(World->PersistentLevel->MapBuildData))
 			{
-				FString NewMapBuildDataName = AssetRenameData.NewName + TEXT("_BuiltData");
-				// Perform rename of MapBuildData before world otherwise original files left behind
-				AssetsToRename.EmplaceAt(0, FAssetRenameDataWithReferencers(FAssetRenameData(World->PersistentLevel->MapBuildData, AssetRenameData.NewPackagePath, NewMapBuildDataName)));
-				AssetsToRename[0].bOnlyFixSoftReferences = AssetRenameData.bOnlyFixSoftReferences;
-				AssetsToRenameLookup.Add(World->PersistentLevel->MapBuildData);
+				// Leave MapBuildData inside the map's package
+				if (World->PersistentLevel->MapBuildData->GetOutermost() != World->GetOutermost())
+				{
+					FString NewMapBuildDataName = AssetRenameData.NewName + TEXT("_BuiltData");
+					// Perform rename of MapBuildData before world otherwise original files left behind
+					AssetsToRename.EmplaceAt(0, FAssetRenameDataWithReferencers(FAssetRenameData(World->PersistentLevel->MapBuildData, AssetRenameData.NewPackagePath, NewMapBuildDataName)));
+					AssetsToRename[0].bOnlyFixSoftReferences = AssetRenameData.bOnlyFixSoftReferences;
+					AssetsToRenameLookup.Add(World->PersistentLevel->MapBuildData);
+				}
 			}
 		}
 
