Thumbnail scrolling optimized, material functions now make calls to load module only once

diff --git a/Engine/Source/Developer/AssetTools/Private/AssetTypeActions/AssetTypeActions_MaterialFunction.cpp b/Engine/Source/Developer/AssetTools/Private/AssetTypeActions/AssetTypeActions_MaterialFunction.cpp
--- a/Engine/Source/Developer/AssetTools/Private/AssetTypeActions/AssetTypeActions_MaterialFunction.cpp
+++ b/Engine/Source/Developer/AssetTools/Private/AssetTypeActions/AssetTypeActions_MaterialFunction.cpp
@@ -170,14 +170,14 @@ UThumbnailInfo* FAssetTypeActions_MaterialFunction::GetThumbnailInfo(UObject* As
 
 UClass* FAssetTypeActions_MaterialFunctionLayer::GetSupportedClass() const
 {
-	IMaterialEditorModule& MaterialEditorModule = FModuleManager::LoadModuleChecked<IMaterialEditorModule>("MaterialEditor");
+	static IMaterialEditorModule& MaterialEditorModule = FModuleManager::LoadModuleChecked<IMaterialEditorModule>("MaterialEditor");
 	UClass* SupportedClass = MaterialEditorModule.MaterialLayersEnabled() ? UMaterialFunctionMaterialLayer::StaticClass() : nullptr;
 	return SupportedClass;
 }
 
 bool FAssetTypeActions_MaterialFunctionLayer::CanFilter()
 {
-	IMaterialEditorModule& MaterialEditorModule = FModuleManager::LoadModuleChecked<IMaterialEditorModule>("MaterialEditor");
+	static IMaterialEditorModule& MaterialEditorModule = FModuleManager::LoadModuleChecked<IMaterialEditorModule>("MaterialEditor");
 	return MaterialEditorModule.MaterialLayersEnabled();
 }
 
@@ -242,14 +242,14 @@ void FAssetTypeActions_MaterialFunctionLayer::ExecuteNewMFI(TArray<TWeakObjectPt
 
 UClass* FAssetTypeActions_MaterialFunctionLayerBlend::GetSupportedClass() const
 {
-	IMaterialEditorModule& MaterialEditorModule = FModuleManager::LoadModuleChecked<IMaterialEditorModule>("MaterialEditor");
+	static IMaterialEditorModule& MaterialEditorModule = FModuleManager::LoadModuleChecked<IMaterialEditorModule>("MaterialEditor");
 	UClass* SupportedClass = MaterialEditorModule.MaterialLayersEnabled() ? UMaterialFunctionMaterialLayerBlend::StaticClass() : nullptr;
 	return SupportedClass;
 }
 
 bool FAssetTypeActions_MaterialFunctionLayerBlend::CanFilter()
 {
-	IMaterialEditorModule& MaterialEditorModule = FModuleManager::LoadModuleChecked<IMaterialEditorModule>("MaterialEditor");
+	static IMaterialEditorModule& MaterialEditorModule = FModuleManager::LoadModuleChecked<IMaterialEditorModule>("MaterialEditor");
 	return MaterialEditorModule.MaterialLayersEnabled();
 }
 
diff --git a/Engine/Source/Developer/AssetTools/Private/AssetTypeActions/AssetTypeActions_MaterialFunctionInstance.cpp b/Engine/Source/Developer/AssetTools/Private/AssetTypeActions/AssetTypeActions_MaterialFunctionInstance.cpp
--- a/Engine/Source/Developer/AssetTools/Private/AssetTypeActions/AssetTypeActions_MaterialFunctionInstance.cpp
+++ b/Engine/Source/Developer/AssetTools/Private/AssetTypeActions/AssetTypeActions_MaterialFunctionInstance.cpp
@@ -79,14 +79,14 @@ void FAssetTypeActions_MaterialFunctionInstance::ExecuteFindParent(TArray<TWeakO
 
 UClass* FAssetTypeActions_MaterialFunctionLayerInstance::GetSupportedClass() const
 {
-	IMaterialEditorModule& MaterialEditorModule = FModuleManager::LoadModuleChecked<IMaterialEditorModule>("MaterialEditor");
+	static IMaterialEditorModule& MaterialEditorModule = FModuleManager::LoadModuleChecked<IMaterialEditorModule>("MaterialEditor");
 	UClass* SupportedClass = MaterialEditorModule.MaterialLayersEnabled() ? UMaterialFunctionMaterialLayerInstance::StaticClass() : nullptr;
 	return SupportedClass;
 }
 
 bool FAssetTypeActions_MaterialFunctionLayerInstance::CanFilter()
 {
-	IMaterialEditorModule& MaterialEditorModule = FModuleManager::LoadModuleChecked<IMaterialEditorModule>("MaterialEditor");
+	static IMaterialEditorModule& MaterialEditorModule = FModuleManager::LoadModuleChecked<IMaterialEditorModule>("MaterialEditor");
 	return MaterialEditorModule.MaterialLayersEnabled();
 }
 
@@ -148,14 +148,14 @@ void FAssetTypeActions_MaterialFunctionLayerInstance::ExecuteNewMFI(TArray<TWeak
 
 UClass* FAssetTypeActions_MaterialFunctionLayerBlendInstance::GetSupportedClass() const
 {
-	IMaterialEditorModule& MaterialEditorModule = FModuleManager::LoadModuleChecked<IMaterialEditorModule>("MaterialEditor");
+	static IMaterialEditorModule& MaterialEditorModule = FModuleManager::LoadModuleChecked<IMaterialEditorModule>("MaterialEditor");
 	UClass* SupportedClass = MaterialEditorModule.MaterialLayersEnabled() ? UMaterialFunctionMaterialLayerBlendInstance::StaticClass() : nullptr;
 	return SupportedClass;
 }
 
 bool FAssetTypeActions_MaterialFunctionLayerBlendInstance::CanFilter()
 {
-	IMaterialEditorModule& MaterialEditorModule = FModuleManager::LoadModuleChecked<IMaterialEditorModule>("MaterialEditor");
+	static IMaterialEditorModule& MaterialEditorModule = FModuleManager::LoadModuleChecked<IMaterialEditorModule>("MaterialEditor");
 	return MaterialEditorModule.MaterialLayersEnabled();
 }
 
