Force feedback thumbnails optimized, avoids loading asset until previewed

diff --git a/Engine/Source/Developer/AssetTools/Private/AssetTypeActions/AssetTypeActions_ForceFeedbackEffect.cpp b/Engine/Source/Developer/AssetTools/Private/AssetTypeActions/AssetTypeActions_ForceFeedbackEffect.cpp
--- a/Engine/Source/Developer/AssetTools/Private/AssetTypeActions/AssetTypeActions_ForceFeedbackEffect.cpp
+++ b/Engine/Source/Developer/AssetTools/Private/AssetTypeActions/AssetTypeActions_ForceFeedbackEffect.cpp
@@ -44,12 +44,36 @@ void FAssetTypeActions_ForceFeedbackEffect::GetActions( const TArray<UObject*>&
 
 bool FAssetTypeActions_ForceFeedbackEffect::IsEffectPlaying(const TArray<TWeakObjectPtr<UForceFeedbackEffect>>& Objects) const
 {
-	for (const TWeakObjectPtr<UForceFeedbackEffect>& EffectPtr : Objects)
+	if (PreviewForceFeedbackEffect.ForceFeedbackEffect)
 	{
-		UForceFeedbackEffect* Effect = EffectPtr.Get();
-		if (Effect && PreviewForceFeedbackEffect.ForceFeedbackEffect == Effect)
+		for (const TWeakObjectPtr<UForceFeedbackEffect>& EffectPtr : Objects)
+		{
+			UForceFeedbackEffect* Effect = EffectPtr.Get();
+			if (Effect && PreviewForceFeedbackEffect.ForceFeedbackEffect == Effect)
+			{
+				return true;
+			}
+		}
+	}
+
+	return false;
+}
+
+bool FAssetTypeActions_ForceFeedbackEffect::IsEffectPlaying(const UForceFeedbackEffect* ForceFeedbackEffect) const
+{
+	return PreviewForceFeedbackEffect.ForceFeedbackEffect && PreviewForceFeedbackEffect.ForceFeedbackEffect == ForceFeedbackEffect;
+}
+
+bool FAssetTypeActions_ForceFeedbackEffect::IsEffectPlaying(const FAssetData& AssetData) const
+{
+	if (PreviewForceFeedbackEffect.ForceFeedbackEffect)
+	{
+		if (PreviewForceFeedbackEffect.ForceFeedbackEffect->GetFName() == AssetData.AssetName)
 		{
-			return true;
+			if (PreviewForceFeedbackEffect.ForceFeedbackEffect->GetOutermost()->GetFName() == AssetData.PackageName)
+			{
+				return true;
+			}
 		}
 	}
 
@@ -70,7 +94,7 @@ void FAssetTypeActions_ForceFeedbackEffect::AssetsActivated( const TArray<UObjec
 			UForceFeedbackEffect* TargetEffect = Cast<UForceFeedbackEffect>(Object);
 			if ( TargetEffect )
 			{
-				// Only target the first valid sound cue
+				// Only target the first valid effect
 				TArray<TWeakObjectPtr<UForceFeedbackEffect>> EffectList;
 				EffectList.Add(MakeWeakObjectPtr(TargetEffect));
 				if (IsEffectPlaying(EffectList))
@@ -137,12 +161,9 @@ void FAssetTypeActions_ForceFeedbackEffect::StopEffect()
 
 TSharedPtr<SWidget> FAssetTypeActions_ForceFeedbackEffect::GetThumbnailOverlay(const FAssetData& AssetData) const
 {
-	TArray<TWeakObjectPtr<UForceFeedbackEffect>> EffectList;
-	EffectList.Add(MakeWeakObjectPtr((UForceFeedbackEffect*)AssetData.GetAsset()));
-
-	auto OnGetDisplayBrushLambda = [this, EffectList]() -> const FSlateBrush*
+	auto OnGetDisplayBrushLambda = [this, AssetData]() -> const FSlateBrush*
 	{
-		if (IsEffectPlaying(EffectList))
+		if (IsEffectPlaying(AssetData))
 		{
 			return FEditorStyle::GetBrush("MediaAsset.AssetActions.Stop.Large");
 		}
@@ -150,28 +171,24 @@ TSharedPtr<SWidget> FAssetTypeActions_ForceFeedbackEffect::GetThumbnailOverlay(c
 		return FEditorStyle::GetBrush("MediaAsset.AssetActions.Play.Large");
 	};
 
-	auto IsEnabledLambda = [this, EffectList]() -> bool
-	{
-		return CanExecutePlayCommand(EffectList);
-	};
-
 	FAssetTypeActions_ForceFeedbackEffect* MutableThis = const_cast<FAssetTypeActions_ForceFeedbackEffect*>(this);
-	auto OnClickedLambda = [MutableThis, EffectList]() -> FReply
+	auto OnClickedLambda = [MutableThis, AssetData]() -> FReply
 	{
-		if (MutableThis->IsEffectPlaying(EffectList))
+		if (MutableThis->IsEffectPlaying(AssetData))
 		{
-			MutableThis->ExecuteStopEffect(EffectList);
+			MutableThis->StopEffect();
 		}
 		else
 		{
-			MutableThis->ExecutePlayEffect(EffectList);
+			// Load and play asset
+			MutableThis->PlayEffect(Cast<UForceFeedbackEffect>(AssetData.GetAsset()));
 		}
 		return FReply::Handled();
 	};
 
-	auto OnToolTipTextLambda = [this, EffectList]() -> FText
+	auto OnToolTipTextLambda = [this, AssetData]() -> FText
 	{
-		if (IsEffectPlaying(EffectList))
+		if (IsEffectPlaying(AssetData))
 		{
 			return LOCTEXT("Thumbnail_StopForceFeedbackToolTip", "Stop selected force feedback effect");
 		}
@@ -184,9 +201,9 @@ TSharedPtr<SWidget> FAssetTypeActions_ForceFeedbackEffect::GetThumbnailOverlay(c
 		.VAlign(VAlign_Center)
 		.Padding(FMargin(2));
 
-	auto OnGetVisibilityLambda = [this, Box, EffectList]() -> EVisibility
+	auto OnGetVisibilityLambda = [this, Box, AssetData]() -> EVisibility
 	{
-		if (Box->IsHovered() || IsEffectPlaying(EffectList))
+		if (Box->IsHovered() || IsEffectPlaying(AssetData))
 		{
 			return EVisibility::Visible;
 		}
@@ -200,7 +217,6 @@ TSharedPtr<SWidget> FAssetTypeActions_ForceFeedbackEffect::GetThumbnailOverlay(c
 		.Cursor(EMouseCursor::Default) // The outer widget can specify a DragHand cursor, so we need to override that here
 		.ForegroundColor(FSlateColor::UseForeground())
 		.IsFocusable(false)
-		.IsEnabled_Lambda(IsEnabledLambda)
 		.OnClicked_Lambda(OnClickedLambda)
 		.Visibility_Lambda(OnGetVisibilityLambda)
 		[
diff --git a/Engine/Source/Developer/AssetTools/Private/AssetTypeActions/AssetTypeActions_ForceFeedbackEffect.h b/Engine/Source/Developer/AssetTools/Private/AssetTypeActions/AssetTypeActions_ForceFeedbackEffect.h
--- a/Engine/Source/Developer/AssetTools/Private/AssetTypeActions/AssetTypeActions_ForceFeedbackEffect.h
+++ b/Engine/Source/Developer/AssetTools/Private/AssetTypeActions/AssetTypeActions_ForceFeedbackEffect.h
@@ -39,6 +39,12 @@ public:
 	/** Return if the specified effect is playing*/
 	bool IsEffectPlaying(const TArray<TWeakObjectPtr<UForceFeedbackEffect>>& Objects) const;
 
+	/** Return if the specified effect is playing*/
+	bool IsEffectPlaying(const UForceFeedbackEffect* ForceFeedbackEffect) const;
+
+	/** Return if the specified asset is playing an effect*/
+	bool IsEffectPlaying(const FAssetData& AssetData) const;
+
 	/** Handler for when PlayEffect is selected */
 	void ExecutePlayEffect(TArray<TWeakObjectPtr<UForceFeedbackEffect>> Objects);
 
@@ -48,12 +54,12 @@ public:
 	/** Returns true if only one effect is selected to play */
 	bool CanExecutePlayCommand(TArray<TWeakObjectPtr<UForceFeedbackEffect>> Objects) const;
 
-private:
-	/** Plays the specified sound wave */
+	/** Plays the specified effect */
 	void PlayEffect(UForceFeedbackEffect* Effect);
 
-	/** Stops any currently playing sounds */
+	/** Stops any currently playing effect */
 	void StopEffect();
 
+private:
 	FPreviewForceFeedbackEffect PreviewForceFeedbackEffect;
 };
