Removed a few places storing module in static variable for safety now that LoadModuleChecked has been optimized

diff --git a/Engine/Source/Developer/AssetTools/Private/AssetTypeActions/AssetTypeActions_MaterialFunction.cpp b/Engine/Source/Developer/AssetTools/Private/AssetTypeActions/AssetTypeActions_MaterialFunction.cpp
--- a/Engine/Source/Developer/AssetTools/Private/AssetTypeActions/AssetTypeActions_MaterialFunction.cpp
+++ b/Engine/Source/Developer/AssetTools/Private/AssetTypeActions/AssetTypeActions_MaterialFunction.cpp
@@ -170,15 +170,13 @@ UThumbnailInfo* FAssetTypeActions_MaterialFunction::GetThumbnailInfo(UObject* As
 
 UClass* FAssetTypeActions_MaterialFunctionLayer::GetSupportedClass() const
 {
-	static IMaterialEditorModule& MaterialEditorModule = FModuleManager::LoadModuleChecked<IMaterialEditorModule>("MaterialEditor");
-	UClass* SupportedClass = MaterialEditorModule.MaterialLayersEnabled() ? UMaterialFunctionMaterialLayer::StaticClass() : nullptr;
+	UClass* SupportedClass = IMaterialEditorModule::Get().MaterialLayersEnabled() ? UMaterialFunctionMaterialLayer::StaticClass() : nullptr;
 	return SupportedClass;
 }
 
 bool FAssetTypeActions_MaterialFunctionLayer::CanFilter()
 {
-	static IMaterialEditorModule& MaterialEditorModule = FModuleManager::LoadModuleChecked<IMaterialEditorModule>("MaterialEditor");
-	return MaterialEditorModule.MaterialLayersEnabled();
+	return IMaterialEditorModule::Get().MaterialLayersEnabled();
 }
 
 
@@ -242,15 +240,13 @@ void FAssetTypeActions_MaterialFunctionLayer::ExecuteNewMFI(TArray<TWeakObjectPt
 
 UClass* FAssetTypeActions_MaterialFunctionLayerBlend::GetSupportedClass() const
 {
-	static IMaterialEditorModule& MaterialEditorModule = FModuleManager::LoadModuleChecked<IMaterialEditorModule>("MaterialEditor");
-	UClass* SupportedClass = MaterialEditorModule.MaterialLayersEnabled() ? UMaterialFunctionMaterialLayerBlend::StaticClass() : nullptr;
+	UClass* SupportedClass = IMaterialEditorModule::Get().MaterialLayersEnabled() ? UMaterialFunctionMaterialLayerBlend::StaticClass() : nullptr;
 	return SupportedClass;
 }
 
 bool FAssetTypeActions_MaterialFunctionLayerBlend::CanFilter()
 {
-	static IMaterialEditorModule& MaterialEditorModule = FModuleManager::LoadModuleChecked<IMaterialEditorModule>("MaterialEditor");
-	return MaterialEditorModule.MaterialLayersEnabled();
+	return IMaterialEditorModule::Get().MaterialLayersEnabled();
 }
 
 void FAssetTypeActions_MaterialFunctionLayerBlend::ExecuteNewMFI(TArray<TWeakObjectPtr<UMaterialFunctionInterface>> Objects)
diff --git a/Engine/Source/Developer/AssetTools/Private/AssetTypeActions/AssetTypeActions_MaterialFunctionInstance.cpp b/Engine/Source/Developer/AssetTools/Private/AssetTypeActions/AssetTypeActions_MaterialFunctionInstance.cpp
--- a/Engine/Source/Developer/AssetTools/Private/AssetTypeActions/AssetTypeActions_MaterialFunctionInstance.cpp
+++ b/Engine/Source/Developer/AssetTools/Private/AssetTypeActions/AssetTypeActions_MaterialFunctionInstance.cpp
@@ -79,15 +79,13 @@ void FAssetTypeActions_MaterialFunctionInstance::ExecuteFindParent(TArray<TWeakO
 
 UClass* FAssetTypeActions_MaterialFunctionLayerInstance::GetSupportedClass() const
 {
-	static IMaterialEditorModule& MaterialEditorModule = FModuleManager::LoadModuleChecked<IMaterialEditorModule>("MaterialEditor");
-	UClass* SupportedClass = MaterialEditorModule.MaterialLayersEnabled() ? UMaterialFunctionMaterialLayerInstance::StaticClass() : nullptr;
+	UClass* SupportedClass = IMaterialEditorModule::Get().MaterialLayersEnabled() ? UMaterialFunctionMaterialLayerInstance::StaticClass() : nullptr;
 	return SupportedClass;
 }
 
 bool FAssetTypeActions_MaterialFunctionLayerInstance::CanFilter()
 {
-	static IMaterialEditorModule& MaterialEditorModule = FModuleManager::LoadModuleChecked<IMaterialEditorModule>("MaterialEditor");
-	return MaterialEditorModule.MaterialLayersEnabled();
+	return IMaterialEditorModule::Get().MaterialLayersEnabled();
 }
 
 void FAssetTypeActions_MaterialFunctionLayerInstance::ExecuteNewMFI(TArray<TWeakObjectPtr<UMaterialFunctionInterface>> Objects)
@@ -148,15 +146,13 @@ void FAssetTypeActions_MaterialFunctionLayerInstance::ExecuteNewMFI(TArray<TWeak
 
 UClass* FAssetTypeActions_MaterialFunctionLayerBlendInstance::GetSupportedClass() const
 {
-	static IMaterialEditorModule& MaterialEditorModule = FModuleManager::LoadModuleChecked<IMaterialEditorModule>("MaterialEditor");
-	UClass* SupportedClass = MaterialEditorModule.MaterialLayersEnabled() ? UMaterialFunctionMaterialLayerBlendInstance::StaticClass() : nullptr;
+	UClass* SupportedClass = IMaterialEditorModule::Get().MaterialLayersEnabled() ? UMaterialFunctionMaterialLayerBlendInstance::StaticClass() : nullptr;
 	return SupportedClass;
 }
 
 bool FAssetTypeActions_MaterialFunctionLayerBlendInstance::CanFilter()
 {
-	static IMaterialEditorModule& MaterialEditorModule = FModuleManager::LoadModuleChecked<IMaterialEditorModule>("MaterialEditor");
-	return MaterialEditorModule.MaterialLayersEnabled();
+	return IMaterialEditorModule::Get().MaterialLayersEnabled();
 }
 
 
diff --git a/Engine/Source/Editor/MaterialEditor/Public/MaterialEditorModule.h b/Engine/Source/Editor/MaterialEditor/Public/MaterialEditorModule.h
--- a/Engine/Source/Editor/MaterialEditor/Public/MaterialEditorModule.h
+++ b/Engine/Source/Editor/MaterialEditor/Public/MaterialEditorModule.h
@@ -63,6 +63,18 @@ public:
 
 	virtual bool MaterialLayersEnabled() = 0;
 
+	/**
+	 * Singleton-like access to this module's interface.  This is just for convenience!
+	 * Beware of calling this during the shutdown phase, though.  Your module might have been unloaded already.
+	 *
+	 * @return Returns singleton instance, loading the module on demand if needed
+	 */
+	static inline IMaterialEditorModule& Get()
+	{
+		static const FName ModuleName("MaterialEditor");
+		return FModuleManager::LoadModuleChecked<IMaterialEditorModule>(ModuleName);
+	}
+
 private:
 	/** All extender delegates for the material menus */
 	TArray<FMaterialMenuExtender> MaterialCanvasMenuExtenders;
