Update copyright notice

diff --git a/Engine/Source/Developer/AssetTools/Private/AssetTypeActions/AssetTypeActions_AnimCurveCompressionSettings.cpp b/Engine/Source/Developer/AssetTools/Private/AssetTypeActions/AssetTypeActions_AnimCurveCompressionSettings.cpp
--- a/Engine/Source/Developer/AssetTools/Private/AssetTypeActions/AssetTypeActions_AnimCurveCompressionSettings.cpp
+++ b/Engine/Source/Developer/AssetTools/Private/AssetTypeActions/AssetTypeActions_AnimCurveCompressionSettings.cpp
@@ -1,4 +1,4 @@
-// Copyright 1998-2018 Epic Games, Inc. All Rights Reserved.
+// Copyright 1998-2019 Epic Games, Inc. All Rights Reserved.
 
 #include "AssetTypeActions/AssetTypeActions_AnimCurveCompressionSettings.h"
 #include "Animation/AnimSequence.h"
diff --git a/Engine/Source/Developer/AssetTools/Private/AssetTypeActions/AssetTypeActions_AnimCurveCompressionSettings.h b/Engine/Source/Developer/AssetTools/Private/AssetTypeActions/AssetTypeActions_AnimCurveCompressionSettings.h
--- a/Engine/Source/Developer/AssetTools/Private/AssetTypeActions/AssetTypeActions_AnimCurveCompressionSettings.h
+++ b/Engine/Source/Developer/AssetTools/Private/AssetTypeActions/AssetTypeActions_AnimCurveCompressionSettings.h
@@ -1,4 +1,4 @@
-// Copyright 1998-2018 Epic Games, Inc. All Rights Reserved.
+// Copyright 1998-2019 Epic Games, Inc. All Rights Reserved.
 
 #pragma once
 
diff --git a/Engine/Source/Editor/UnrealEd/Classes/Factories/AnimCurveCompressionSettingsFactory.h b/Engine/Source/Editor/UnrealEd/Classes/Factories/AnimCurveCompressionSettingsFactory.h
--- a/Engine/Source/Editor/UnrealEd/Classes/Factories/AnimCurveCompressionSettingsFactory.h
+++ b/Engine/Source/Editor/UnrealEd/Classes/Factories/AnimCurveCompressionSettingsFactory.h
@@ -1,4 +1,4 @@
-// Copyright 1998-2018 Epic Games, Inc. All Rights Reserved.
+// Copyright 1998-2019 Epic Games, Inc. All Rights Reserved.
 
 #pragma once
 
diff --git a/Engine/Source/Editor/UnrealEd/Private/Factories/AnimCurveCompressionSettingsFactory.cpp b/Engine/Source/Editor/UnrealEd/Private/Factories/AnimCurveCompressionSettingsFactory.cpp
--- a/Engine/Source/Editor/UnrealEd/Private/Factories/AnimCurveCompressionSettingsFactory.cpp
+++ b/Engine/Source/Editor/UnrealEd/Private/Factories/AnimCurveCompressionSettingsFactory.cpp
@@ -1,4 +1,4 @@
-// Copyright 1998-2018 Epic Games, Inc. All Rights Reserved.
+// Copyright 1998-2019 Epic Games, Inc. All Rights Reserved.
 
 /*=============================================================================
 AnimCurveCompressionSettingsFactory.cpp: Factory for animation curve compression settings assets
diff --git a/Engine/Source/Runtime/Engine/Classes/Animation/AnimCurveCompressionCodec.h b/Engine/Source/Runtime/Engine/Classes/Animation/AnimCurveCompressionCodec.h
--- a/Engine/Source/Runtime/Engine/Classes/Animation/AnimCurveCompressionCodec.h
+++ b/Engine/Source/Runtime/Engine/Classes/Animation/AnimCurveCompressionCodec.h
@@ -1,4 +1,4 @@
-// Copyright 1998-2018 Epic Games, Inc. All Rights Reserved.
+// Copyright 1998-2019 Epic Games, Inc. All Rights Reserved.
 
 #pragma once
 
diff --git a/Engine/Source/Runtime/Engine/Classes/Animation/AnimCurveCompressionCodec_CompressedRichCurve.h b/Engine/Source/Runtime/Engine/Classes/Animation/AnimCurveCompressionCodec_CompressedRichCurve.h
--- a/Engine/Source/Runtime/Engine/Classes/Animation/AnimCurveCompressionCodec_CompressedRichCurve.h
+++ b/Engine/Source/Runtime/Engine/Classes/Animation/AnimCurveCompressionCodec_CompressedRichCurve.h
@@ -1,4 +1,4 @@
-// Copyright 1998-2018 Epic Games, Inc. All Rights Reserved.
+// Copyright 1998-2019 Epic Games, Inc. All Rights Reserved.
 
 #pragma once
 
diff --git a/Engine/Source/Runtime/Engine/Classes/Animation/AnimCurveCompressionCodec_UniformlySampled.h b/Engine/Source/Runtime/Engine/Classes/Animation/AnimCurveCompressionCodec_UniformlySampled.h
--- a/Engine/Source/Runtime/Engine/Classes/Animation/AnimCurveCompressionCodec_UniformlySampled.h
+++ b/Engine/Source/Runtime/Engine/Classes/Animation/AnimCurveCompressionCodec_UniformlySampled.h
@@ -1,4 +1,4 @@
-// Copyright 1998-2018 Epic Games, Inc. All Rights Reserved.
+// Copyright 1998-2019 Epic Games, Inc. All Rights Reserved.
 
 #pragma once
 
diff --git a/Engine/Source/Runtime/Engine/Classes/Animation/AnimCurveCompressionSettings.h b/Engine/Source/Runtime/Engine/Classes/Animation/AnimCurveCompressionSettings.h
--- a/Engine/Source/Runtime/Engine/Classes/Animation/AnimCurveCompressionSettings.h
+++ b/Engine/Source/Runtime/Engine/Classes/Animation/AnimCurveCompressionSettings.h
@@ -1,4 +1,4 @@
-// Copyright 1998-2018 Epic Games, Inc. All Rights Reserved.
+// Copyright 1998-2019 Epic Games, Inc. All Rights Reserved.
 
 #pragma once
 
diff --git a/Engine/Source/Runtime/Engine/Private/Animation/AnimCurveCompressionCodec.cpp b/Engine/Source/Runtime/Engine/Private/Animation/AnimCurveCompressionCodec.cpp
--- a/Engine/Source/Runtime/Engine/Private/Animation/AnimCurveCompressionCodec.cpp
+++ b/Engine/Source/Runtime/Engine/Private/Animation/AnimCurveCompressionCodec.cpp
@@ -1,4 +1,4 @@
-// Copyright 1998-2018 Epic Games, Inc. All Rights Reserved.
+// Copyright 1998-2019 Epic Games, Inc. All Rights Reserved.
 
 #include "Animation/AnimCurveCompressionCodec.h"
 
diff --git a/Engine/Source/Runtime/Engine/Private/Animation/AnimCurveCompressionCodec_CompressedRichCurve.cpp b/Engine/Source/Runtime/Engine/Private/Animation/AnimCurveCompressionCodec_CompressedRichCurve.cpp
--- a/Engine/Source/Runtime/Engine/Private/Animation/AnimCurveCompressionCodec_CompressedRichCurve.cpp
+++ b/Engine/Source/Runtime/Engine/Private/Animation/AnimCurveCompressionCodec_CompressedRichCurve.cpp
@@ -1,4 +1,4 @@
-// Copyright 1998-2018 Epic Games, Inc. All Rights Reserved.
+// Copyright 1998-2019 Epic Games, Inc. All Rights Reserved.
 
 #include "Animation/AnimCurveCompressionCodec_CompressedRichCurve.h"
 #include "Animation/AnimSequence.h"
diff --git a/Engine/Source/Runtime/Engine/Private/Animation/AnimCurveCompressionSettings.cpp b/Engine/Source/Runtime/Engine/Private/Animation/AnimCurveCompressionSettings.cpp
--- a/Engine/Source/Runtime/Engine/Private/Animation/AnimCurveCompressionSettings.cpp
+++ b/Engine/Source/Runtime/Engine/Private/Animation/AnimCurveCompressionSettings.cpp
@@ -1,4 +1,4 @@
-// Copyright 1998-2018 Epic Games, Inc. All Rights Reserved.
+// Copyright 1998-2019 Epic Games, Inc. All Rights Reserved.
 
 #include "Animation/AnimCurveCompressionSettings.h"
 #include "Animation/AnimCurveCompressionCodec_CompressedRichCurve.h"
