Allow supplying factory to asset import task

diff --git a/Engine/Source/Developer/AssetTools/Private/AssetTools.cpp b/Engine/Source/Developer/AssetTools/Private/AssetTools.cpp
--- a/Engine/Source/Developer/AssetTools/Private/AssetTools.cpp
+++ b/Engine/Source/Developer/AssetTools/Private/AssetTools.cpp
@@ -1104,6 +1104,7 @@ void UAssetToolsImpl::ImportAssetTasks(const TArray<UAssetImportTask*>& ImportTa
 		Params.AssetImportTask = ImportTask;
 		Params.bForceOverrideExisting = ImportTask->bReplaceExisting;
 		Params.bAutomated = ImportTask->bAutomated;
+		Params.SpecifiedFactory = ImportTask->Factory;
 		Filenames[0] = ImportTask->Filename;
 		TArray<UObject*> ImportedObjects = ImportAssetsInternal(Filenames, ImportTask->DestinationPath, nullptr, Params);
 
diff --git a/Engine/Source/Editor/UnrealEd/Public/AssetImportTask.h b/Engine/Source/Editor/UnrealEd/Public/AssetImportTask.h
--- a/Engine/Source/Editor/UnrealEd/Public/AssetImportTask.h
+++ b/Engine/Source/Editor/UnrealEd/Public/AssetImportTask.h
@@ -45,6 +45,10 @@ public:
 	UPROPERTY(BlueprintReadWrite, Category = "Asset Import Task")
 	bool bSave;
 
+	/** Optional factory to use */
+	UPROPERTY(BlueprintReadWrite, Category = "Asset Import Task")
+	UFactory* Factory;
+
 	/** Import options specific to the type of asset */
 	UPROPERTY(BlueprintReadWrite, Category = "Asset Import Task")
 	UObject* Options;
