Folder advanced copy now works, preserves structure.

diff --git a/Engine/Source/Developer/AssetTools/Private/AdvancedCopyCustomization.cpp b/Engine/Source/Developer/AssetTools/Private/AdvancedCopyCustomization.cpp
--- a/Engine/Source/Developer/AssetTools/Private/AdvancedCopyCustomization.cpp
+++ b/Engine/Source/Developer/AssetTools/Private/AdvancedCopyCustomization.cpp
@@ -2,6 +2,7 @@
 
 #include "AdvancedCopyCustomization.h"
 #include "Containers/UnrealString.h"
+#include "AssetRegistryModule.h"
 
 
 #define LOCTEXT_NAMESPACE "AdvancedCopyCustomization"
@@ -20,6 +21,18 @@ UAdvancedCopyCustomization::UAdvancedCopyCustomization(const class FObjectInitia
 void UAdvancedCopyCustomization::SetPackageThatInitiatedCopy(const FString& InBasePackage)
 {
 	FString TempPackage = InBasePackage;
+
+	FAssetRegistryModule& AssetRegistryModule = FModuleManager::Get().LoadModuleChecked<FAssetRegistryModule>(TEXT("AssetRegistry"));
+	TArray<FAssetData> DependencyAssetData;
+	IAssetRegistry& AssetRegistry = AssetRegistryModule.Get();
+	AssetRegistry.GetAssetsByPackageName(FName(*InBasePackage), DependencyAssetData);
+	// We found a folder
+	if (DependencyAssetData.Num() == 0)
+	{
+		// Take off the name of the folder we copied so copied files are still nested
+		TempPackage.Split(TEXT("/"), &TempPackage, nullptr, ESearchCase::IgnoreCase, ESearchDir::FromEnd);
+	}
+	
 	if (!TempPackage.EndsWith(TEXT("/")))
 	{
 		TempPackage += TEXT("/");
diff --git a/Engine/Source/Developer/AssetTools/Private/AssetTools.cpp b/Engine/Source/Developer/AssetTools/Private/AssetTools.cpp
--- a/Engine/Source/Developer/AssetTools/Private/AssetTools.cpp
+++ b/Engine/Source/Developer/AssetTools/Private/AssetTools.cpp
@@ -816,7 +816,11 @@ void UAssetToolsImpl::GetAllAdvancedCopySources(FName SelectedPackage, FAdvanced
 	if (!OutPackageNamesToCopy.Contains(SelectedPackage))
 	{
 		CurrentDependencies.Add(SelectedPackage);
-		if (CopyParams.bShouldCheckForDependencies)
+		TArray<FAssetData> DependencyAssetData;
+		AssetRegistry.GetAssetsByPackageName(SelectedPackage, DependencyAssetData);
+		// If we should check for dependencies OR we are currently checking a folder
+		// Folders should ALWAYS get checked for assets and subfolders
+		if (CopyParams.bShouldCheckForDependencies || DependencyAssetData.Num() == 0)
 		{
 			RecursiveGetDependenciesAdvanced(SelectedPackage, CopyParams, CurrentDependencies, DependencyMap);
 		}
@@ -2836,26 +2840,60 @@ void UAssetToolsImpl::RecursiveGetDependenciesAdvanced(const FName& PackageName,
 {
 	FAssetRegistryModule& AssetRegistryModule = FModuleManager::Get().LoadModuleChecked<FAssetRegistryModule>(TEXT("AssetRegistry"));
 	TArray<FName> Dependencies;
-	AssetRegistryModule.Get().GetDependencies(PackageName, Dependencies);
+	TArray<FAssetData> DependencyAssetData;
 	IAssetRegistry& AssetRegistry = AssetRegistryModule.Get();
+	AssetRegistry.GetAssetsByPackageName(PackageName, DependencyAssetData);
+	// We found an asset
+	if (DependencyAssetData.Num() > 0)
+	{
+		AssetRegistryModule.Get().GetDependencies(PackageName, Dependencies);
+		for (auto DependsIt = Dependencies.CreateConstIterator(); DependsIt; ++DependsIt)
+		{
+			if (!AllDependencies.Contains(*DependsIt))
+			{
+				if (AssetRegistry.GetAssetsByPackageName(*DependsIt, DependencyAssetData))
+				{
+					FARFilter ExclusionFilter = UAdvancedCopyCustomization::StaticClass()->GetDefaultObject<UAdvancedCopyCustomization>()->GetARFilter();
+					AssetRegistry.UseFilterToExcludeAssets(DependencyAssetData, ExclusionFilter);
+					if (DependencyAssetData.IsValidIndex(0))
+					{
+						AllDependencies.Add(*DependsIt);
+						DependencyMap.Add(*DependsIt, PackageName);
+						RecursiveGetDependenciesAdvanced(*DependsIt, CopyParams, AllDependencies, DependencyMap);
+					}
+				}
 
-	for (auto DependsIt = Dependencies.CreateConstIterator(); DependsIt; ++DependsIt)
+			}
+		}
+	}
+	else
 	{
-		if (!AllDependencies.Contains(*DependsIt))
+		DependencyAssetData.Empty();
+		// We found a folder containing assets
+		if (AssetRegistry.HasAssets(PackageName) && AssetRegistry.GetAssetsByPath(PackageName, DependencyAssetData))
 		{
-			TArray<FAssetData> DependencyAssetData;
-			if (AssetRegistry.GetAssetsByPackageName(*DependsIt, DependencyAssetData))
+			FARFilter ExclusionFilter = UAdvancedCopyCustomization::StaticClass()->GetDefaultObject<UAdvancedCopyCustomization>()->GetARFilter();
+			AssetRegistry.UseFilterToExcludeAssets(DependencyAssetData, ExclusionFilter);
+			for(const FAssetData Asset : DependencyAssetData)
 			{
-				FARFilter ExclusionFilter = UAdvancedCopyCustomization::StaticClass()->GetDefaultObject<UAdvancedCopyCustomization>()->GetARFilter();
-				AssetRegistry.UseFilterToExcludeAssets(DependencyAssetData, ExclusionFilter);
-				if (DependencyAssetData.IsValidIndex(0))
+				AllDependencies.Add(*Asset.GetPackage()->GetName());
+				// If we should check the assets we found for dependencies
+				if (CopyParams.bShouldCheckForDependencies)
 				{
-					AllDependencies.Add(*DependsIt);
-					DependencyMap.Add(*DependsIt, PackageName);
-					RecursiveGetDependenciesAdvanced(*DependsIt, CopyParams, AllDependencies, DependencyMap);
+					RecursiveGetDependenciesAdvanced(FName(*Asset.GetPackage()->GetName()), CopyParams, AllDependencies, DependencyMap);
 				}
 			}
+		}
 		
+		// Always get subpaths
+		{
+			TArray<FString> SubPaths;
+			AssetRegistry.GetSubPaths(PackageName.ToString(), SubPaths, false);
+			for (const FString SubPath : SubPaths)
+			{
+				RecursiveGetDependenciesAdvanced(FName(*SubPath), CopyParams, AllDependencies, DependencyMap);
+			}
+
 		}
 	}
 }
@@ -2870,7 +2908,7 @@ void UAssetToolsImpl::OpenEditorForAssets(const TArray<UObject*>& Assets) const
 	FAssetEditorManager::Get().OpenEditorForAssets(Assets);
 }
 
-void UAssetToolsImpl::BeginAdvancedCopyPackages(const TArray<FName>& PackageNamesToCopy, const FString& TargetPath) const
+void UAssetToolsImpl::BeginAdvancedCopyPackages(const TArray<FName>& InputNamesToCopy, const FString& TargetPath) const
 {
 	// Packages must be saved for the migration to work
 	const bool bPromptUserToSave = true;
@@ -2884,34 +2922,35 @@ void UAssetToolsImpl::BeginAdvancedCopyPackages(const TArray<FName>& PackageName
 		{
 			// Open a dialog asking the user to wait while assets are being discovered
 			SDiscoveringAssetsDialog::OpenDiscoveringAssetsDialog(
-				SDiscoveringAssetsDialog::FOnAssetsDiscovered::CreateUObject(this, &UAssetToolsImpl::PerformAdvancedCopyPackages, PackageNamesToCopy, TargetPath)
+				SDiscoveringAssetsDialog::FOnAssetsDiscovered::CreateUObject(this, &UAssetToolsImpl::PerformAdvancedCopyPackages, InputNamesToCopy, TargetPath)
 			);
 		}
 		else
 		{
 			// Assets are already discovered, perform the migration now
-			PerformAdvancedCopyPackages(PackageNamesToCopy, TargetPath);
+			PerformAdvancedCopyPackages(InputNamesToCopy, TargetPath);
 		}
 	}
 }
 
-void UAssetToolsImpl::PerformAdvancedCopyPackages(TArray<FName> SelectedPackageNames, FString TargetPath) const
+void UAssetToolsImpl::PerformAdvancedCopyPackages(TArray<FName> SelectedAssetAndFolderNames, FString TargetPath) const
 {	
-	FAdvancedCopyParams CopyParams = FAdvancedCopyParams(SelectedPackageNames, TargetPath);
-	CopyParams.bShouldCheckForDependencies = SelectedPackageNames.Num() == 1;
+	TargetPath.RemoveFromEnd(TEXT("/"));
+	FAdvancedCopyParams CopyParams = FAdvancedCopyParams(SelectedAssetAndFolderNames, TargetPath);
+	CopyParams.bShouldCheckForDependencies = SelectedAssetAndFolderNames.Num() == 1;
 
-	for (auto PackageIt = SelectedPackageNames.CreateConstIterator(); PackageIt; ++PackageIt)
+	for (auto NameIt = SelectedAssetAndFolderNames.CreateConstIterator(); NameIt; ++NameIt)
 	{
 		UAdvancedCopyCustomization* CopyCustomization = nullptr;
 
 		const UAssetToolsSettings* Settings = GetDefault<UAssetToolsSettings>();
-		FName OriginalPackage = *PackageIt;
-		const FString& PackageNameString = OriginalPackage.ToString();
+		FName OriginalName = *NameIt;
+		const FString& OriginalNameString = OriginalName.ToString();
 		FString SrcFilename;
 		UObject* ExistingObject = nullptr;
-		if (FPackageName::DoesPackageExist(PackageNameString, nullptr, &SrcFilename))
+		if (FPackageName::DoesPackageExist(OriginalNameString, nullptr, &SrcFilename))
 		{
-			UPackage* Pkg = LoadPackage(nullptr, *PackageNameString, LOAD_None);
+			UPackage* Pkg = LoadPackage(nullptr, *OriginalNameString, LOAD_None);
 			if (Pkg)
 			{
 				FString Name = ObjectTools::SanitizeObjectName(FPaths::GetBaseFilename(SrcFilename));
@@ -2935,6 +2974,7 @@ void UAssetToolsImpl::PerformAdvancedCopyPackages(TArray<FName> SelectedPackageN
 		}
 
 		// If not able to find class in settings, fall back to default customization
+		// by default, folders will use the default customization
 		if (CopyCustomization == nullptr)
 		{
 			CopyCustomization = UAdvancedCopyCustomization::StaticClass()->GetDefaultObject<UAdvancedCopyCustomization>();
diff --git a/Engine/Source/Developer/AssetTools/Private/AssetTools.h b/Engine/Source/Developer/AssetTools/Private/AssetTools.h
--- a/Engine/Source/Developer/AssetTools/Private/AssetTools.h
+++ b/Engine/Source/Developer/AssetTools/Private/AssetTools.h
@@ -96,7 +96,7 @@ public:
 	virtual FString DumpAssetToTempFile(UObject* Asset) const override;
 	virtual bool CreateDiffProcess(const FString& DiffCommand, const FString& OldTextFilename, const FString& NewTextFilename, const FString& DiffArgs = FString("")) const override;
 	virtual void MigratePackages(const TArray<FName>& PackageNamesToMigrate) const override;
-	virtual void BeginAdvancedCopyPackages(const TArray<FName>& PackageNamesToCopy, const FString& TargetPath) const override;
+	virtual void BeginAdvancedCopyPackages(const TArray<FName>& InputNamesToCopy, const FString& TargetPath) const override;
 	virtual void FixupReferencers(const TArray<UObjectRedirector*>& Objects) const override;
 	virtual FAssetPostRenameEvent& OnAssetPostRename() override { return AssetRenameManager->OnAssetPostRenameEvent(); }
 	virtual void ExpandDirectories(const TArray<FString>& Files, const FString& DestinationPath, TArray<TPair<FString, FString>>& FilesAndDestinations) const override;
diff --git a/Engine/Source/Developer/AssetTools/Public/IAssetTools.h b/Engine/Source/Developer/AssetTools/Public/IAssetTools.h
--- a/Engine/Source/Developer/AssetTools/Public/IAssetTools.h
+++ b/Engine/Source/Developer/AssetTools/Public/IAssetTools.h
@@ -368,7 +368,7 @@ public:
 	virtual void MigratePackages(const TArray<FName>& PackageNamesToMigrate) const = 0;
 
 	/* Copy packages and dependencies to another folder */
-	virtual void BeginAdvancedCopyPackages(const TArray<FName>& PackageNamesToCopy, const FString& TargetPath) const = 0;
+	virtual void BeginAdvancedCopyPackages(const TArray<FName>& InputNamesToCopy, const FString& TargetPath) const = 0;
 
 	/** Fix up references to the specified redirectors */
 	virtual void FixupReferencers(const TArray<UObjectRedirector*>& Objects) const = 0;
diff --git a/Engine/Source/Editor/ContentBrowser/Private/ContentBrowserUtils.cpp b/Engine/Source/Editor/ContentBrowser/Private/ContentBrowserUtils.cpp
--- a/Engine/Source/Editor/ContentBrowser/Private/ContentBrowserUtils.cpp
+++ b/Engine/Source/Editor/ContentBrowser/Private/ContentBrowserUtils.cpp
@@ -2324,4 +2324,29 @@ int32 ContentBrowserUtils::GetMaxCookPathLen()
 	}
 }
 
+void ContentBrowserUtils::BeginAdvancedCopyPackages(TArray<FAssetData>& AssetList, TArray<FString>& AssetPaths, FString& DestinationPath)
+{
+	TMap<FString, TArray<UObject*> > SourcePathToLoadedAssets;
+	FString DestPathWithTrailingSlash = DestinationPath / "";
+	// Get a list of package names for input into Advanced Copy
+	TArray<FName> InputNames;
+	// Do not allow parent directories to be moved to themselves or children.
+	TArray<FString> SourcePathNames = AssetPaths;
+
+	for (int32 AssetIdx = 0; AssetIdx < AssetList.Num(); ++AssetIdx)
+	{
+		InputNames.Add(AssetList[AssetIdx].PackageName);
+	}
+
+	// Add any paths from selected folders
+	for (const FString AssetPath : AssetPaths)
+	{
+		InputNames.Add(FName(*AssetPath));
+
+	}
+
+	FAssetToolsModule& AssetToolsModule = FModuleManager::Get().LoadModuleChecked<FAssetToolsModule>("AssetTools");
+	AssetToolsModule.Get().BeginAdvancedCopyPackages(InputNames, DestPathWithTrailingSlash);
+}
+
 #undef LOCTEXT_NAMESPACE
diff --git a/Engine/Source/Editor/ContentBrowser/Private/ContentBrowserUtils.h b/Engine/Source/Editor/ContentBrowser/Private/ContentBrowserUtils.h
--- a/Engine/Source/Editor/ContentBrowser/Private/ContentBrowserUtils.h
+++ b/Engine/Source/Editor/ContentBrowser/Private/ContentBrowserUtils.h
@@ -283,5 +283,7 @@ namespace ContentBrowserUtils
 
 	/** Gets the maximum path length for a cooked file. Changes behavior based on whether the editor experimental setting for long paths is enabled. */
 	int32 GetMaxCookPathLen();
+
+	void BeginAdvancedCopyPackages(TArray<FAssetData>& AssetList, TArray<FString>& AssetPaths, FString& DestinationPath);
 }
 
diff --git a/Engine/Source/Editor/ContentBrowser/Private/SAssetView.cpp b/Engine/Source/Editor/ContentBrowser/Private/SAssetView.cpp
--- a/Engine/Source/Editor/ContentBrowser/Private/SAssetView.cpp
+++ b/Engine/Source/Editor/ContentBrowser/Private/SAssetView.cpp
@@ -4878,18 +4878,7 @@ void SAssetView::ExecuteDropMove(TArray<FAssetData> AssetList, TArray<FString> A
 
 void SAssetView::ExecuteDropAdvancedCopy(TArray<FAssetData> AssetList, TArray<FString> AssetPaths, FString DestinationPath)
 {
-	int32 NumItemsCopied = 0;
-	// Get a list of package names for input into Advanced Copy 
-	TArray<FName> PackageNames;
-	PackageNames.Reserve(AssetList.Num());
-
-	for (int32 AssetIdx = 0; AssetIdx < AssetList.Num(); ++AssetIdx)
-	{
-		PackageNames.Add(AssetList[AssetIdx].PackageName);
-	}
-
-	FAssetToolsModule& AssetToolsModule = FModuleManager::Get().LoadModuleChecked<FAssetToolsModule>("AssetTools");
-	AssetToolsModule.Get().BeginAdvancedCopyPackages(PackageNames, DestinationPath);
+	ContentBrowserUtils::BeginAdvancedCopyPackages(AssetList, AssetPaths, DestinationPath);
 }
 
 void SAssetView::SetUserSearching(bool bInSearching)
diff --git a/Engine/Source/Editor/ContentBrowser/Private/SPathView.cpp b/Engine/Source/Editor/ContentBrowser/Private/SPathView.cpp
--- a/Engine/Source/Editor/ContentBrowser/Private/SPathView.cpp
+++ b/Engine/Source/Editor/ContentBrowser/Private/SPathView.cpp
@@ -1491,18 +1491,7 @@ void SPathView::ExecuteTreeDropMove(TArray<FAssetData> AssetList, TArray<FString
 
 void SPathView::ExecuteTreeDropAdvancedCopy(TArray<FAssetData> AssetList, TArray<FString> AssetPaths, FString DestinationPath)
 {
-	int32 NumItemsCopied = 0;
-	// Get a list of package names for input into Advanced Copy
-	TArray<FName> PackageNames;
-	PackageNames.Reserve(AssetList.Num());
-
-	for (int32 AssetIdx = 0; AssetIdx < AssetList.Num(); ++AssetIdx)
-	{
-		PackageNames.Add(AssetList[AssetIdx].PackageName);
-	}
-
-	FAssetToolsModule& AssetToolsModule = FModuleManager::Get().LoadModuleChecked<FAssetToolsModule>("AssetTools");
-	AssetToolsModule.Get().BeginAdvancedCopyPackages(PackageNames, DestinationPath);
+	ContentBrowserUtils::BeginAdvancedCopyPackages(AssetList, AssetPaths, DestinationPath);
 }
 
 void SPathView::OnAssetRegistryPathAdded(const FString& Path)
