Fix for excluding engine dependencies, not hiding copied asset, being able to advanced copy more than once
#rnx

diff --git a/Engine/Source/Developer/AssetTools/Private/AssetTools.cpp b/Engine/Source/Developer/AssetTools/Private/AssetTools.cpp
--- a/Engine/Source/Developer/AssetTools/Private/AssetTools.cpp
+++ b/Engine/Source/Developer/AssetTools/Private/AssetTools.cpp
@@ -926,7 +926,6 @@ bool UAssetToolsImpl::AdvancedCopyPackages(const TMap<FString, FString>& SourceA
 
 		FMessageLog AdvancedCopyLog("AssetTools");
 		FText LogMessage = FText::FromString(TEXT("Advanced content copy completed successfully!"));
-		FModuleManager::LoadModuleChecked<FAssetRegistryModule>(TEXT("AssetRegistry")).Get().ScanModifiedAssetFiles(SuccessfullyCopiedPackages);
 		EMessageSeverity::Type Severity = EMessageSeverity::Info;
 		if (SourceControlErrors.Len() > 0)
 		{
@@ -2840,17 +2839,18 @@ void UAssetToolsImpl::RecursiveGetDependenciesAdvanced(const FName& PackageName,
 {
 	FAssetRegistryModule& AssetRegistryModule = FModuleManager::Get().LoadModuleChecked<FAssetRegistryModule>(TEXT("AssetRegistry"));
 	TArray<FName> Dependencies;
-	TArray<FAssetData> DependencyAssetData;
+	TArray<FAssetData> PackageAssetData;
 	IAssetRegistry& AssetRegistry = AssetRegistryModule.Get();
-	AssetRegistry.GetAssetsByPackageName(PackageName, DependencyAssetData);
+	AssetRegistry.GetAssetsByPackageName(PackageName, PackageAssetData);
 	// We found an asset
-	if (DependencyAssetData.Num() > 0)
+	if (PackageAssetData.Num() > 0)
 	{
 		AssetRegistryModule.Get().GetDependencies(PackageName, Dependencies);
 		for (auto DependsIt = Dependencies.CreateConstIterator(); DependsIt; ++DependsIt)
 		{
 			if (!AllDependencies.Contains(*DependsIt))
 			{
+				TArray<FAssetData> DependencyAssetData;
 				if (AssetRegistry.GetAssetsByPackageName(*DependsIt, DependencyAssetData))
 				{
 					FARFilter ExclusionFilter = UAdvancedCopyCustomization::StaticClass()->GetDefaultObject<UAdvancedCopyCustomization>()->GetARFilter();
@@ -2868,13 +2868,13 @@ void UAssetToolsImpl::RecursiveGetDependenciesAdvanced(const FName& PackageName,
 	}
 	else
 	{
-		DependencyAssetData.Empty();
+		TArray<FAssetData> PathAssetData;
 		// We found a folder containing assets
-		if (AssetRegistry.HasAssets(PackageName) && AssetRegistry.GetAssetsByPath(PackageName, DependencyAssetData))
+		if (AssetRegistry.HasAssets(PackageName) && AssetRegistry.GetAssetsByPath(PackageName, PathAssetData))
 		{
 			FARFilter ExclusionFilter = UAdvancedCopyCustomization::StaticClass()->GetDefaultObject<UAdvancedCopyCustomization>()->GetARFilter();
-			AssetRegistry.UseFilterToExcludeAssets(DependencyAssetData, ExclusionFilter);
-			for(const FAssetData Asset : DependencyAssetData)
+			AssetRegistry.UseFilterToExcludeAssets(PathAssetData, ExclusionFilter);
+			for(const FAssetData Asset : PathAssetData)
 			{
 				AllDependencies.Add(*Asset.GetPackage()->GetName());
 				// If we should check the assets we found for dependencies
