Exclude non-project plugin content from advanced copy

diff --git a/Engine/Source/Developer/AssetTools/AssetTools.Build.cs b/Engine/Source/Developer/AssetTools/AssetTools.Build.cs
--- a/Engine/Source/Developer/AssetTools/AssetTools.Build.cs
+++ b/Engine/Source/Developer/AssetTools/AssetTools.Build.cs
@@ -30,8 +30,9 @@ public class AssetTools : ModuleRules
 				"PropertyEditor",
 				"Kismet",
 				"Landscape",
-                "Foliage"
-			}
+                "Foliage",
+                "Projects",
+            }
 		);
 
 		PrivateIncludePathModuleNames.AddRange(
diff --git a/Engine/Source/Developer/AssetTools/Private/AdvancedCopyCustomization.cpp b/Engine/Source/Developer/AssetTools/Private/AdvancedCopyCustomization.cpp
--- a/Engine/Source/Developer/AssetTools/Private/AdvancedCopyCustomization.cpp
+++ b/Engine/Source/Developer/AssetTools/Private/AdvancedCopyCustomization.cpp
@@ -3,6 +3,7 @@
 #include "AdvancedCopyCustomization.h"
 #include "Containers/UnrealString.h"
 #include "AssetRegistryModule.h"
+#include "Interfaces/IPluginManager.h"
 
 
 #define LOCTEXT_NAMESPACE "AdvancedCopyCustomization"
@@ -13,7 +14,14 @@ UAdvancedCopyCustomization::UAdvancedCopyCustomization(const class FObjectInitia
 	, bShouldGenerateRelativePaths(true)
 {
 	FilterForExcludingDependencies.PackagePaths.Add(TEXT("/Engine"));
-	FilterForExcludingDependencies.PackagePaths.Add(TEXT("/Script"));
+	for (TSharedRef<IPlugin>& Plugin : IPluginManager::Get().GetDiscoveredPlugins())
+	{
+		if (Plugin->GetType() != EPluginType::Project)
+		{
+			FilterForExcludingDependencies.PackagePaths.Add(FName(*("/" + Plugin->GetName())));
+		}
+	}
+
 	FilterForExcludingDependencies.bRecursivePaths = true;
 	FilterForExcludingDependencies.bRecursiveClasses = true;
 }
