Remove deprecated data table export flags in engine

diff --git a/Engine/Source/Developer/AssetTools/Private/AssetTypeActions/AssetTypeActions_DataTable.cpp b/Engine/Source/Developer/AssetTools/Private/AssetTypeActions/AssetTypeActions_DataTable.cpp
--- a/Engine/Source/Developer/AssetTools/Private/AssetTypeActions/AssetTypeActions_DataTable.cpp
+++ b/Engine/Source/Developer/AssetTools/Private/AssetTypeActions/AssetTypeActions_DataTable.cpp
@@ -91,7 +91,7 @@ void FAssetTypeActions_DataTable::ExecuteExportAsCSV(TArray< TWeakObjectPtr<UObj
 
 			if (OutFilenames.Num() > 0)
 			{
-				FFileHelper::SaveStringToFile(DataTable->GetTableAsCSV(EDataTableExportFlags::UsePrettyPropertyNames | EDataTableExportFlags::UsePrettyEnumNames), *OutFilenames[0]);
+				FFileHelper::SaveStringToFile(DataTable->GetTableAsCSV(), *OutFilenames[0]);
 			}
 		}
 	}
@@ -125,7 +125,7 @@ void FAssetTypeActions_DataTable::ExecuteExportAsJSON(TArray< TWeakObjectPtr<UOb
 
 			if (OutFilenames.Num() > 0)
 			{
-				FFileHelper::SaveStringToFile(DataTable->GetTableAsJSON(EDataTableExportFlags::UsePrettyPropertyNames | EDataTableExportFlags::UsePrettyEnumNames | EDataTableExportFlags::UseJsonObjectsForStructs), *OutFilenames[0]);
+				FFileHelper::SaveStringToFile(DataTable->GetTableAsJSON(EDataTableExportFlags::UseJsonObjectsForStructs), *OutFilenames[0]);
 			}
 		}
 	}
@@ -209,8 +209,8 @@ void FAssetTypeActions_DataTable::PerformAssetDiff(UObject* OldAsset, UObject* N
 	FString AbsoluteNewTempFileName = FPaths::ConvertRelativePathToFull(RelNewTempFileName);
 
 	// save temp files
-	bool OldResult = FFileHelper::SaveStringToFile(OldDataTable->GetTableAsCSV(EDataTableExportFlags::UsePrettyPropertyNames | EDataTableExportFlags::UsePrettyEnumNames), *AbsoluteOldTempFileName);
-	bool NewResult = FFileHelper::SaveStringToFile(NewDataTable->GetTableAsCSV(EDataTableExportFlags::UsePrettyPropertyNames | EDataTableExportFlags::UsePrettyEnumNames), *AbsoluteNewTempFileName);
+	bool OldResult = FFileHelper::SaveStringToFile(OldDataTable->GetTableAsCSV(), *AbsoluteOldTempFileName);
+	bool NewResult = FFileHelper::SaveStringToFile(NewDataTable->GetTableAsCSV(), *AbsoluteNewTempFileName);
 
 	if (OldResult && NewResult)
 	{
