Fixed crash when a redirector is displayed using a skeletal mesh FAssetTypeActions on asset rename

diff --git a/Engine/Source/Developer/AssetTools/Private/AssetTypeActions/AssetTypeActions_SkeletalMesh.cpp b/Engine/Source/Developer/AssetTools/Private/AssetTypeActions/AssetTypeActions_SkeletalMesh.cpp
--- a/Engine/Source/Developer/AssetTools/Private/AssetTypeActions/AssetTypeActions_SkeletalMesh.cpp
+++ b/Engine/Source/Developer/AssetTools/Private/AssetTypeActions/AssetTypeActions_SkeletalMesh.cpp
@@ -588,15 +588,16 @@ EVisibility FAssetTypeActions_SkeletalMesh::GetThumbnailSkinningOverlayVisibilit
 
 	//The object is loaded we can use the memory value of the object to set the overlay
 	UObject* Obj = AssetData.GetAsset();
-	USkeletalMesh* SkeletalMesh = CastChecked<USkeletalMesh>(Obj);
-
-	UAssetImportData* GenericImportData = SkeletalMesh->AssetImportData;
-	if (GenericImportData != nullptr)
+	if(USkeletalMesh* SkeletalMesh = Cast<USkeletalMesh>(Obj))
 	{
-		UFbxSkeletalMeshImportData* ImportData = Cast<UFbxSkeletalMeshImportData>(GenericImportData);
-		if (ImportData != nullptr && ImportData->LastImportContentType == EFBXImportContentType::FBXICT_Geometry)
+		UAssetImportData* GenericImportData = SkeletalMesh->AssetImportData;
+		if (GenericImportData != nullptr)
 		{
-			return EVisibility::HitTestInvisible;
+			UFbxSkeletalMeshImportData* ImportData = Cast<UFbxSkeletalMeshImportData>(GenericImportData);
+			if (ImportData != nullptr && ImportData->LastImportContentType == EFBXImportContentType::FBXICT_Geometry)
+			{
+				return EVisibility::HitTestInvisible;
+			}
 		}
 	}
 	return EVisibility::Collapsed;
