Assets: Blueprint callable functions must be non const

diff --git a/Engine/Source/Developer/AssetTools/Private/AssetTools.cpp b/Engine/Source/Developer/AssetTools/Private/AssetTools.cpp
--- a/Engine/Source/Developer/AssetTools/Private/AssetTools.cpp
+++ b/Engine/Source/Developer/AssetTools/Private/AssetTools.cpp
@@ -998,22 +998,22 @@ bool UAssetToolsImpl::AdvancedCopyPackages(const FAdvancedCopyParams& CopyParams
 	return false;
 }
 
-bool UAssetToolsImpl::RenameAssets(const TArray<FAssetRenameData>& AssetsAndNames) const
+bool UAssetToolsImpl::RenameAssets(const TArray<FAssetRenameData>& AssetsAndNames)
 {
 	return AssetRenameManager->RenameAssets(AssetsAndNames);
 }
 
-void UAssetToolsImpl::RenameAssetsWithDialog(const TArray<FAssetRenameData>& AssetsAndNames, bool bAutoCheckout) const
+void UAssetToolsImpl::RenameAssetsWithDialog(const TArray<FAssetRenameData>& AssetsAndNames, bool bAutoCheckout)
 {
 	AssetRenameManager->RenameAssetsWithDialog(AssetsAndNames, bAutoCheckout);
 }
 
-void UAssetToolsImpl::FindSoftReferencesToObject(FSoftObjectPath TargetObject, TArray<UObject*>& ReferencingObjects) const
+void UAssetToolsImpl::FindSoftReferencesToObject(FSoftObjectPath TargetObject, TArray<UObject*>& ReferencingObjects)
 {
 	AssetRenameManager->FindSoftReferencesToObject(TargetObject, ReferencingObjects);
 }
 
-void UAssetToolsImpl::RenameReferencingSoftObjectPaths(const TArray<UPackage *> PackagesToCheck, const TMap<FSoftObjectPath, FSoftObjectPath>& AssetRedirectorMap) const
+void UAssetToolsImpl::RenameReferencingSoftObjectPaths(const TArray<UPackage *> PackagesToCheck, const TMap<FSoftObjectPath, FSoftObjectPath>& AssetRedirectorMap)
 {
 	AssetRenameManager->RenameReferencingSoftObjectPaths(PackagesToCheck, AssetRedirectorMap);
 }
@@ -1092,7 +1092,7 @@ TArray<UObject*> UAssetToolsImpl::ImportAssetsWithDialog(const FString& Destinat
 	return ReturnObjects;
 }
 
-TArray<UObject*> UAssetToolsImpl::ImportAssetsAutomated(const UAutomatedAssetImportData* ImportData) const
+TArray<UObject*> UAssetToolsImpl::ImportAssetsAutomated(const UAutomatedAssetImportData* ImportData)
 {
 	check(ImportData);
 
@@ -1107,7 +1107,7 @@ TArray<UObject*> UAssetToolsImpl::ImportAssetsAutomated(const UAutomatedAssetImp
 	return ImportAssetsInternal(ImportData->Filenames, ImportData->DestinationPath, nullptr, Params);
 }
 
-void UAssetToolsImpl::ImportAssetTasks(const TArray<UAssetImportTask*>& ImportTasks) const
+void UAssetToolsImpl::ImportAssetTasks(const TArray<UAssetImportTask*>& ImportTasks)
 {
 	FScopedSlowTask SlowTask(ImportTasks.Num(), LOCTEXT("ImportSlowTask", "Importing"));
 	SlowTask.MakeDialog();
@@ -1152,7 +1152,7 @@ void UAssetToolsImpl::ImportAssetTasks(const TArray<UAssetImportTask*>& ImportTa
 	}
 }
 
-void UAssetToolsImpl::ExportAssets(const TArray<FString>& AssetsToExport, const FString& ExportPath) const
+void UAssetToolsImpl::ExportAssets(const TArray<FString>& AssetsToExport, const FString& ExportPath)
 {
 	TArray<UObject*> AssetObjectsToExport;
 	AssetObjectsToExport.Reserve(AssetsToExport.Num());
@@ -1180,12 +1180,12 @@ void UAssetToolsImpl::ExportAssets(const TArray<UObject*>& AssetsToExport, const
 	ExportAssetsInternal(AssetsToExport, bPromptIndividualFilenames, ExportPath);
 }
 
-void UAssetToolsImpl::ExportAssetsWithDialog(const TArray<UObject*>& AssetsToExport, bool bPromptForIndividualFilenames) const
+void UAssetToolsImpl::ExportAssetsWithDialog(const TArray<UObject*>& AssetsToExport, bool bPromptForIndividualFilenames)
 {
 	ExportAssetsInternal(AssetsToExport, bPromptForIndividualFilenames, TEXT(""));
 }
 
-void UAssetToolsImpl::ExportAssetsWithDialog(const TArray<FString>& AssetsToExport, bool bPromptForIndividualFilenames) const
+void UAssetToolsImpl::ExportAssetsWithDialog(const TArray<FString>& AssetsToExport, bool bPromptForIndividualFilenames)
 {
 	TArray<UObject*> AssetObjectsToExport;
 	AssetObjectsToExport.Reserve(AssetsToExport.Num());
@@ -1255,7 +1255,7 @@ TArray<UObject*> UAssetToolsImpl::ImportAssets(const TArray<FString>& Files, con
 	return ImportAssetsInternal(Files, DestinationPath, FilesAndDestinations, Params);
 }
 
-void UAssetToolsImpl::CreateUniqueAssetName(const FString& InBasePackageName, const FString& InSuffix, FString& OutPackageName, FString& OutAssetName) const
+void UAssetToolsImpl::CreateUniqueAssetName(const FString& InBasePackageName, const FString& InSuffix, FString& OutPackageName, FString& OutAssetName)
 {
 	const FString SanitizedBasePackageName = UPackageTools::SanitizePackageName(InBasePackageName);
 
@@ -2920,7 +2920,7 @@ void UAssetToolsImpl::FixupReferencers(const TArray<UObjectRedirector*>& Objects
 	AssetFixUpRedirectors->FixupReferencers(Objects);
 }
 
-void UAssetToolsImpl::OpenEditorForAssets(const TArray<UObject*>& Assets) const
+void UAssetToolsImpl::OpenEditorForAssets(const TArray<UObject*>& Assets)
 {
 	FAssetEditorManager::Get().OpenEditorForAssets(Assets);
 }
diff --git a/Engine/Source/Developer/AssetTools/Private/AssetTools.h b/Engine/Source/Developer/AssetTools/Private/AssetTools.h
--- a/Engine/Source/Developer/AssetTools/Private/AssetTools.h
+++ b/Engine/Source/Developer/AssetTools/Private/AssetTools.h
@@ -75,21 +75,21 @@ public:
 	virtual UObject* CreateAssetWithDialog(const FString& AssetName, const FString& PackagePath, UClass* AssetClass, UFactory* Factory, FName CallingContext = NAME_None) override;
 	virtual UObject* DuplicateAsset(const FString& AssetName, const FString& PackagePath, UObject* OriginalObject) override;
 	virtual UObject* DuplicateAssetWithDialog(const FString& AssetName, const FString& PackagePath, UObject* OriginalObject) override;
-	virtual bool RenameAssets(const TArray<FAssetRenameData>& AssetsAndNames) const override;
-	virtual void RenameAssetsWithDialog(const TArray<FAssetRenameData>& AssetsAndNames, bool bAutoCheckout = false) const override;
-	virtual void FindSoftReferencesToObject(FSoftObjectPath TargetObject, TArray<UObject*>& ReferencingObjects) const override;
-	virtual void RenameReferencingSoftObjectPaths(const TArray<UPackage *> PackagesToCheck, const TMap<FSoftObjectPath, FSoftObjectPath>& AssetRedirectorMap) const override;
+	virtual bool RenameAssets(const TArray<FAssetRenameData>& AssetsAndNames) override;
+	virtual void RenameAssetsWithDialog(const TArray<FAssetRenameData>& AssetsAndNames, bool bAutoCheckout = false) override;
+	virtual void FindSoftReferencesToObject(FSoftObjectPath TargetObject, TArray<UObject*>& ReferencingObjects) override;
+	virtual void RenameReferencingSoftObjectPaths(const TArray<UPackage *> PackagesToCheck, const TMap<FSoftObjectPath, FSoftObjectPath>& AssetRedirectorMap) override;
 	virtual TArray<UObject*> ImportAssets(const FString& DestinationPath) override;
 	virtual TArray<UObject*> ImportAssetsWithDialog(const FString& DestinationPath) override;
 	virtual TArray<UObject*> ImportAssets(const TArray<FString>& Files, const FString& DestinationPath, UFactory* ChosenFactory, bool bSyncToBrowser = true, TArray<TPair<FString, FString>>* FilesAndDestinations = nullptr) const override;
-	virtual TArray<UObject*> ImportAssetsAutomated(const UAutomatedAssetImportData* ImportData) const override;
-	virtual void ImportAssetTasks(const TArray<UAssetImportTask*>& ImportTasks) const override;
-	virtual void ExportAssets(const TArray<FString>& AssetsToExport, const FString& ExportPath) const override;
+	virtual TArray<UObject*> ImportAssetsAutomated(const UAutomatedAssetImportData* ImportData) override;
+	virtual void ImportAssetTasks(const TArray<UAssetImportTask*>& ImportTasks) override;
+	virtual void ExportAssets(const TArray<FString>& AssetsToExport, const FString& ExportPath) override;
 	virtual void ExportAssets(const TArray<UObject*>& AssetsToExport, const FString& ExportPath) const override;
-	virtual void ExportAssetsWithDialog(const TArray<UObject*>& AssetsToExport, bool bPromptForIndividualFilenames) const override;
-	virtual void ExportAssetsWithDialog(const TArray<FString>& AssetsToExport, bool bPromptForIndividualFilenames) const override;
+	virtual void ExportAssetsWithDialog(const TArray<UObject*>& AssetsToExport, bool bPromptForIndividualFilenames) override;
+	virtual void ExportAssetsWithDialog(const TArray<FString>& AssetsToExport, bool bPromptForIndividualFilenames) override;
 
-	virtual void CreateUniqueAssetName(const FString& InBasePackageName, const FString& InSuffix, FString& OutPackageName, FString& OutAssetName) const override;
+	virtual void CreateUniqueAssetName(const FString& InBasePackageName, const FString& InSuffix, FString& OutPackageName, FString& OutAssetName) override;
 	virtual bool AssetUsesGenericThumbnail( const FAssetData& AssetData ) const override;
 	virtual void DiffAgainstDepot(UObject* InObject, const FString& InPackagePath, const FString& InPackageName) const override;
 	virtual void DiffAssets(UObject* OldAsset1, UObject* NewAsset, const struct FRevisionInfo& OldRevision, const struct FRevisionInfo& NewRevision) const override;
@@ -108,7 +108,7 @@ public:
 	virtual void GetAllAdvancedCopySources(FName SelectedPackage, FAdvancedCopyParams& CopyParams, TArray<FName>& OutPackageNamesToCopy, TMap<FName, FName>& DependencyMap, const class UAdvancedCopyCustomization* CopyCustomization) const override;
 	virtual void InitAdvancedCopyFromCopyParams(FAdvancedCopyParams CopyParams) const override;
 
-	virtual void OpenEditorForAssets(const TArray<UObject*>& Assets) const override;
+	virtual void OpenEditorForAssets(const TArray<UObject*>& Assets) override;
 public:
 	/** Gets the asset tools singleton as a FAssetTools for asset tools module use */
 	static UAssetToolsImpl& Get();
diff --git a/Engine/Source/Developer/AssetTools/Public/IAssetTools.h b/Engine/Source/Developer/AssetTools/Public/IAssetTools.h
--- a/Engine/Source/Developer/AssetTools/Public/IAssetTools.h
+++ b/Engine/Source/Developer/AssetTools/Public/IAssetTools.h
@@ -236,15 +236,15 @@ public:
 
 	/** Renames assets using the specified names. */
 	UFUNCTION(BlueprintCallable, Category = "Editor Scripting | Asset Tools")
-	virtual bool RenameAssets(const TArray<FAssetRenameData>& AssetsAndNames) const = 0;
+	virtual bool RenameAssets(const TArray<FAssetRenameData>& AssetsAndNames) = 0;
 
 	/** Renames assets using the specified names. */
 	UFUNCTION(BlueprintCallable, Category = "Editor Scripting | Asset Tools")
-	virtual void RenameAssetsWithDialog(const TArray<FAssetRenameData>& AssetsAndNames, bool bAutoCheckout = false) const = 0;
+	virtual void RenameAssetsWithDialog(const TArray<FAssetRenameData>& AssetsAndNames, bool bAutoCheckout = false) = 0;
 
 	/** Returns list of objects that soft reference the given soft object path. This will load assets into memory to verify */
 	UFUNCTION(BlueprintCallable, Category = "Editor Scripting | Asset Tools")
-	virtual void FindSoftReferencesToObject(FSoftObjectPath TargetObject, TArray<UObject*>& ReferencingObjects) const = 0;
+	virtual void FindSoftReferencesToObject(FSoftObjectPath TargetObject, TArray<UObject*>& ReferencingObjects) = 0;
 
 	/**
 	 * Function that renames all FSoftObjectPath object with the old asset path to the new one.
@@ -252,7 +252,8 @@ public:
 	 * @param PackagesToCheck Packages to check for referencing FSoftObjectPath.
 	 * @param AssetRedirectorMap Map from old asset path to new asset path
 	 */
-	virtual void RenameReferencingSoftObjectPaths(const TArray<UPackage *> PackagesToCheck, const TMap<FSoftObjectPath, FSoftObjectPath>& AssetRedirectorMap) const = 0;
+	UFUNCTION(BlueprintCallable, Category = "Editor Scripting | Asset Tools")
+	virtual void RenameReferencingSoftObjectPaths(const TArray<UPackage *> PackagesToCheck, const TMap<FSoftObjectPath, FSoftObjectPath>& AssetRedirectorMap) = 0;
 
 	/** Event issued at the end of the rename process */
 	virtual FAssetPostRenameEvent& OnAssetPostRename() = 0;
@@ -287,7 +288,7 @@ public:
 	 * @return list of successfully imported assets
 	 */
 	UFUNCTION(BlueprintCallable, Category = "Editor Scripting | Asset Tools")
-	virtual TArray<UObject*> ImportAssetsAutomated( const UAutomatedAssetImportData* ImportData) const = 0;
+	virtual TArray<UObject*> ImportAssetsAutomated( const UAutomatedAssetImportData* ImportData) = 0;
 
 	/**
 	* Imports assets using tasks specified.
@@ -295,7 +296,7 @@ public:
 	* @param ImportTasks	Tasks that specify how to import each file
 	*/
 	UFUNCTION(BlueprintCallable, Category = "Editor Scripting | Asset Tools")
-	virtual void ImportAssetTasks(const TArray<UAssetImportTask*>& ImportTasks) const = 0;
+	virtual void ImportAssetTasks(const TArray<UAssetImportTask*>& ImportTasks) = 0;
 
 	/**
 	 * Exports the specified objects to file.
@@ -304,7 +305,7 @@ public:
 	 * @param	ExportPath						The directory path to export to.
 	 */
 	UFUNCTION(BlueprintCallable, Category = "Editor Scripting | Asset Tools")
-	virtual void ExportAssets(const TArray<FString>& AssetsToExport, const FString& ExportPath) const = 0;
+	virtual void ExportAssets(const TArray<FString>& AssetsToExport, const FString& ExportPath) = 0;
 
 	/**
 	 * Exports the specified objects to file.
@@ -321,7 +322,7 @@ public:
 	 * @param	ExportPath						The directory path to export to.
 	 */
 	UFUNCTION(BlueprintCallable, Category = "Editor Scripting | Asset Tools")
-	virtual void ExportAssetsWithDialog(const TArray<FString>& AssetsToExport, bool bPromptForIndividualFilenames) const = 0;
+	virtual void ExportAssetsWithDialog(const TArray<FString>& AssetsToExport, bool bPromptForIndividualFilenames) = 0;
 	
 	/**
 	 * Exports the specified objects to file. First prompting the user to pick an export directory and optionally prompting the user to pick a unique directory per file
@@ -329,11 +330,11 @@ public:
 	 * @param	AssetsToExport					List of full asset names (e.g /Game/Path/Asset) to export
 	 * @param	ExportPath						The directory path to export to.
 	 */
-	virtual void ExportAssetsWithDialog(const TArray<UObject*>& AssetsToExport, bool bPromptForIndividualFilenames) const = 0;
+	virtual void ExportAssetsWithDialog(const TArray<UObject*>& AssetsToExport, bool bPromptForIndividualFilenames) = 0;
 
 	/** Creates a unique package and asset name taking the form InBasePackageName+InSuffix */
 	UFUNCTION(BlueprintCallable, Category = "Editor Scripting | Asset Tools")
-	virtual void CreateUniqueAssetName(const FString& InBasePackageName, const FString& InSuffix, FString& OutPackageName, FString& OutAssetName) const = 0;
+	virtual void CreateUniqueAssetName(const FString& InBasePackageName, const FString& InSuffix, FString& OutPackageName, FString& OutAssetName) = 0;
 
 	/** Returns true if the specified asset uses a stock thumbnail resource */
 	virtual bool AssetUsesGenericThumbnail(const FAssetData& AssetData) const = 0;
@@ -399,7 +400,7 @@ public:
 
 	/** Opens editor for assets */
 	UFUNCTION(BlueprintCallable, Category = "Editor Scripting | Asset Tools")
-	virtual void OpenEditorForAssets(const TArray<UObject*>& Assets) const = 0;
+	virtual void OpenEditorForAssets(const TArray<UObject*>& Assets) = 0;
 };
 
 UCLASS(transient)
