Integrating CL 7700278

Fixing opacity mask using Virtual Textures at Datasmith import

diff --git a/Engine/Source/Developer/AssetTools/Private/AssetTools.cpp b/Engine/Source/Developer/AssetTools/Private/AssetTools.cpp
--- a/Engine/Source/Developer/AssetTools/Private/AssetTools.cpp
+++ b/Engine/Source/Developer/AssetTools/Private/AssetTools.cpp
@@ -127,6 +127,7 @@
 #include "AdvancedCopyCustomization.h"
 #include "SAdvancedCopyReportDialog.h"
 #include "AssetToolsSettings.h"
+#include "AssetVtConversion.h"
 
 
 #define LOCTEXT_NAMESPACE "AssetTools"
@@ -2929,6 +2930,20 @@ void UAssetToolsImpl::OpenEditorForAssets(const TArray<UObject*>& Assets)
 	FAssetEditorManager::Get().OpenEditorForAssets(Assets);
 }
 
+void UAssetToolsImpl::ConvertVirtualTextures(const TArray<UTexture2D *>& Textures, bool bConvertBackToNonVirtual, const TArray<UMaterial *>* RelatedMaterials /* = nullptr */) const
+{
+	FVTConversionWorker VirtualTextureConversionWorker(bConvertBackToNonVirtual);
+	VirtualTextureConversionWorker.UserTextures = Textures;
+	//We want all given texture to be added, so we put a minimum texture size of 0
+	VirtualTextureConversionWorker.FilterList(0);
+	if (RelatedMaterials)
+	{
+		VirtualTextureConversionWorker.Materials.Append(*RelatedMaterials);
+	}
+
+	VirtualTextureConversionWorker.DoConvert();
+}
+
 void UAssetToolsImpl::BeginAdvancedCopyPackages(const TArray<FName>& InputNamesToCopy, const FString& TargetPath) const
 {
 	// Packages must be saved for the migration to work
diff --git a/Engine/Source/Developer/AssetTools/Private/AssetTools.h b/Engine/Source/Developer/AssetTools/Private/AssetTools.h
--- a/Engine/Source/Developer/AssetTools/Private/AssetTools.h
+++ b/Engine/Source/Developer/AssetTools/Private/AssetTools.h
@@ -112,6 +112,8 @@ public:
 	virtual void InitAdvancedCopyFromCopyParams(FAdvancedCopyParams CopyParams) const override;
 
 	virtual void OpenEditorForAssets(const TArray<UObject*>& Assets) override;
+	
+	virtual void ConvertVirtualTextures(const TArray<UTexture2D*>& Textures, bool bConvertBackToNonVirtual, const TArray<UMaterial*>* RelatedMaterials = nullptr) const override;
 public:
 	/** Gets the asset tools singleton as a FAssetTools for asset tools module use */
 	static UAssetToolsImpl& Get();
diff --git a/Engine/Source/Developer/AssetTools/Private/AssetTypeActions/AssetTypeActions_Texture.cpp b/Engine/Source/Developer/AssetTools/Private/AssetTypeActions/AssetTypeActions_Texture.cpp
--- a/Engine/Source/Developer/AssetTools/Private/AssetTypeActions/AssetTypeActions_Texture.cpp
+++ b/Engine/Source/Developer/AssetTools/Private/AssetTypeActions/AssetTypeActions_Texture.cpp
@@ -396,6 +396,7 @@ public:
 
 	void SetBackwards(bool bSetBackwards)
 	{
+		Worker.SetConversionDirection(bBackwards);
 		bBackwards = bSetBackwards;
 	}
 
@@ -713,7 +714,7 @@ private:
 
 		if (ButtonID == FConvertToVTDlg::EResult::Confirm)
 		{
-			Worker.DoConvert(bBackwards);
+			Worker.DoConvert();
 		}
 
 		return FReply::Handled();
@@ -859,8 +860,8 @@ FConvertToVTDlg::FConvertToVTDlg(const TArray<UTexture2D *> &Textures, bool bBac
 				.ParentWindow(DialogWindow)
 			];
 
-		DialogWidget->SetUserTextures(Textures);
 		DialogWidget->SetBackwards(bBackwards);
+		DialogWidget->SetUserTextures(Textures);
 		DialogWindow->SetContent(DialogWrapper.ToSharedRef());
 	}
 }
diff --git a/Engine/Source/Developer/AssetTools/Private/AssetVtConversion.cpp b/Engine/Source/Developer/AssetTools/Private/AssetVtConversion.cpp
--- a/Engine/Source/Developer/AssetTools/Private/AssetVtConversion.cpp
+++ b/Engine/Source/Developer/AssetTools/Private/AssetVtConversion.cpp
@@ -157,7 +157,7 @@ void FVTConversionWorker::FindAllTexturesAndMaterials_Iteration(TArray<UMaterial
 		bool bIsVirtualTextureValid = true;
 		for (UMaterialInterface* Material : MaterialsUsingTexture)
 		{
-			if (!IsVirtualTextureValidForMaterial(Material, Tex2D))
+			if (!bConvertBackward && !IsVirtualTextureValidForMaterial(Material, Tex2D))
 			{
 				AuditTrail.Add(Tex2D, FAuditTrail(
 					Material,
@@ -516,9 +516,9 @@ static void MarkTextureExpressionModified(UMaterialExpressionTextureBase* Expres
 	Expression->PostEditChangeProperty(TextureChangeEvent);
 }
 
-void FVTConversionWorker::DoConvert(bool backwards)
+void FVTConversionWorker::DoConvert()
 {
-	const bool bVirtualTextureEnable = !backwards;
+	const bool bVirtualTextureEnable = !bConvertBackward;
 
 	FScopedSlowTask SlowTask(2.0f, LOCTEXT("ConvertToVT_Progress_ConvertingTexturesAndMaterials", "Converting textures and materials..."));
 	SlowTask.MakeDialog();
diff --git a/Engine/Source/Developer/AssetTools/Private/AssetVtConversion.h b/Engine/Source/Developer/AssetTools/Private/AssetVtConversion.h
--- a/Engine/Source/Developer/AssetTools/Private/AssetVtConversion.h
+++ b/Engine/Source/Developer/AssetTools/Private/AssetVtConversion.h
@@ -13,6 +13,10 @@ struct FScopedSlowTask;
 class FVTConversionWorker
 {
 public:
+	FVTConversionWorker(bool bInConvertBackward = false) : bConvertBackward(bInConvertBackward){}
+
+	bool bConvertBackward;
+
 	// Fill this array intially with some textures you want to see converted to vt
 	TArray<UTexture2D*> UserTextures;
 
@@ -32,10 +36,12 @@ public:
 
 	TMap<UObject*, FAuditTrail> AuditTrail;
 
+	void SetConversionDirection(bool bInConvertBackward) { bConvertBackward = bInConvertBackward; }
+
 	void FilterList(int32 SizeThreshold);
 
 	// Based on the filtered lists do the actual VT conversion
-	void DoConvert(bool backwards);
+	void DoConvert();
 private:
 	void FindAllTexturesAndMaterials_Iteration(TArray<UMaterial*>& InAffectedMaterials,
 		TArray<UMaterialFunctionInterface*>& InAffectedFunctions,
diff --git a/Engine/Source/Developer/AssetTools/Public/IAssetTools.h b/Engine/Source/Developer/AssetTools/Public/IAssetTools.h
--- a/Engine/Source/Developer/AssetTools/Public/IAssetTools.h
+++ b/Engine/Source/Developer/AssetTools/Public/IAssetTools.h
@@ -401,6 +401,13 @@ public:
 	/** Opens editor for assets */
 	UFUNCTION(BlueprintCallable, Category = "Editor Scripting | Asset Tools")
 	virtual void OpenEditorForAssets(const TArray<UObject*>& Assets) = 0;
+	
+	/** Converts the given UTexture2D to virtual textures or converts virtual textures back to standard textures and updates the related UMaterials 
+	 * @param Textures					The given textures to convert.
+	 * @param bConvertBackToNonVirtual	If true, virtual textures will be converted back to standard textures.
+	 * @param RelatedMaterials			An optional array of materials to update after the conversion, this is useful during import when not all dependencies and caches are up to date.
+	 */
+	virtual void ConvertVirtualTextures(const TArray<UTexture2D*>& Textures, bool bConvertBackToNonVirtual, const TArray<UMaterial*>* RelatedMaterials = nullptr) const = 0;
 };
 
 UCLASS(transient)
diff --git a/Engine/Source/Runtime/Engine/Private/Materials/HLSLMaterialTranslator.h b/Engine/Source/Runtime/Engine/Private/Materials/HLSLMaterialTranslator.h
--- a/Engine/Source/Runtime/Engine/Private/Materials/HLSLMaterialTranslator.h
+++ b/Engine/Source/Runtime/Engine/Private/Materials/HLSLMaterialTranslator.h
@@ -778,11 +778,14 @@ public:
 				if (Visitor.FoundVirtualTexture)
 				{
 					// virtual texture connected to an invalid property, report the correct error
+#if MATERIAL_OPACITYMASK_DOESNT_SUPPORT_VIRTUALTEXTURE
 					if (PropertyToValidate == MP_OpacityMask)
 					{
 						Errorf(TEXT("Sampling a virtual texture is currently not supported when connected to the Opacity Mask material attribute."));
 					}
-					else if (ShaderFrequencyToValidate != SF_Pixel)
+					else
+#endif
+					if (ShaderFrequencyToValidate != SF_Pixel)
 					{
 						Errorf(TEXT("Sampling a virtual texture is currently only supported from pixel shader."));
 					}
diff --git a/Engine/Source/Runtime/Engine/Public/MaterialShared.h b/Engine/Source/Runtime/Engine/Public/MaterialShared.h
--- a/Engine/Source/Runtime/Engine/Public/MaterialShared.h
+++ b/Engine/Source/Runtime/Engine/Public/MaterialShared.h
@@ -73,6 +73,8 @@ template <class ElementType> class TLinkedList;
 #define STORE_ONLY_ACTIVE_SHADERMAPS 0
 #endif
 
+#define MATERIAL_OPACITYMASK_DOESNT_SUPPORT_VIRTUALTEXTURE 1
+
 DECLARE_LOG_CATEGORY_EXTERN(LogMaterial,Log,Verbose);
 
 /** Creates a string that represents the given quality level. */
