Integrating CL #7794913 from Dev-Enterprise

Fixing Virtual Texture conversion tool

diff --git a/Engine/Source/Developer/AssetTools/Private/AssetTypeActions/AssetTypeActions_Texture.cpp b/Engine/Source/Developer/AssetTools/Private/AssetTypeActions/AssetTypeActions_Texture.cpp
--- a/Engine/Source/Developer/AssetTools/Private/AssetTypeActions/AssetTypeActions_Texture.cpp
+++ b/Engine/Source/Developer/AssetTools/Private/AssetTypeActions/AssetTypeActions_Texture.cpp
@@ -396,7 +396,7 @@ public:
 
 	void SetBackwards(bool bSetBackwards)
 	{
-		Worker.SetConversionDirection(bBackwards);
+		Worker.SetConversionDirection(bSetBackwards);
 		bBackwards = bSetBackwards;
 	}
 
