Make the content browser import LOD context menu more precise.

diff --git a/Engine/Source/Developer/AssetTools/Private/AssetTypeActions/AssetTypeActions_SkeletalMesh.cpp b/Engine/Source/Developer/AssetTools/Private/AssetTypeActions/AssetTypeActions_SkeletalMesh.cpp
--- a/Engine/Source/Developer/AssetTools/Private/AssetTypeActions/AssetTypeActions_SkeletalMesh.cpp
+++ b/Engine/Source/Developer/AssetTools/Private/AssetTypeActions/AssetTypeActions_SkeletalMesh.cpp
@@ -662,7 +662,7 @@ void FAssetTypeActions_SkeletalMesh::GetLODMenu(class FMenuBuilder& MenuBuilder,
 	int32 LODMax = SkeletalMesh->GetLODNum();
 	for(int32 LOD = 1; LOD <= LODMax; ++LOD)
 	{
-		const FText Description = FText::Format( LOCTEXT("LODLevel", "LOD {0}"), FText::AsNumber( LOD ) );
+		const FText Description = (LOD == LODMax) ? FText::Format(LOCTEXT("AddLODLevel", "Add LOD {0}"), FText::AsNumber(LOD)) : FText::Format( LOCTEXT("LODLevel", "Reimport LOD {0}"), FText::AsNumber( LOD ) );
 		const FText ToolTip = ( LOD == LODMax ) ? LOCTEXT("NewImportTip", "Import new LOD") : LOCTEXT("ReimportTip", "Reimport over existing LOD");
 		MenuBuilder.AddMenuEntry(	Description, 
 									ToolTip, FSlateIcon(),
diff --git a/Engine/Source/Developer/AssetTools/Private/AssetTypeActions/AssetTypeActions_StaticMesh.cpp b/Engine/Source/Developer/AssetTools/Private/AssetTypeActions/AssetTypeActions_StaticMesh.cpp
--- a/Engine/Source/Developer/AssetTools/Private/AssetTypeActions/AssetTypeActions_StaticMesh.cpp
+++ b/Engine/Source/Developer/AssetTools/Private/AssetTypeActions/AssetTypeActions_StaticMesh.cpp
@@ -111,7 +111,7 @@ void FAssetTypeActions_StaticMesh::GetImportLODMenu(class FMenuBuilder& MenuBuil
 		FText ToolTip = NSLOCTEXT("AssetTypeActions_StaticMesh", "ReimportTip", "Reimport over existing LOD");
 		if(LOD == First->GetNumLODs())
 		{
-			Description = FText::Format( NSLOCTEXT("AssetTypeActions_StaticMesh", "LOD (number)", "LOD {0}"), LODText );
+			Description = FText::Format( NSLOCTEXT("AssetTypeActions_StaticMesh", "LOD (number)", "Add LOD {0}"), LODText );
 			ToolTip = NSLOCTEXT("AssetTypeActions_StaticMesh", "NewImportTip", "Import new LOD");
 		}
 
