Extend the Migrate tool to be able to selectively exclude individual files, or entire folders.

diff --git a/Engine/Source/Developer/AssetTools/Private/AssetTools.cpp b/Engine/Source/Developer/AssetTools/Private/AssetTools.cpp
--- a/Engine/Source/Developer/AssetTools/Private/AssetTools.cpp
+++ b/Engine/Source/Developer/AssetTools/Private/AssetTools.cpp
@@ -1553,7 +1553,7 @@ void UAssetToolsImpl::MigratePackages(const TArray<FName>& PackageNamesToMigrate
 			// Open a dialog asking the user to wait while assets are being discovered
 			SDiscoveringAssetsDialog::OpenDiscoveringAssetsDialog(
 				SDiscoveringAssetsDialog::FOnAssetsDiscovered::CreateUObject(this, &UAssetToolsImpl::PerformMigratePackages, PackageNamesToMigrate)
-				);
+			);
 		}
 		else
 		{
@@ -2559,17 +2559,17 @@ void UAssetToolsImpl::PerformMigratePackages(TArray<FName> PackageNamesToMigrate
 	// Prompt the user displaying all assets that are going to be migrated
 	{
 		const FText ReportMessage = LOCTEXT("MigratePackagesReportTitle", "The following assets will be migrated to another content folder.");
-		TArray<FString> ReportPackageNames;
+		TSharedPtr<TArray<ReportPackageData>> ReportPackages = MakeShareable(new TArray<ReportPackageData>);
 		for ( auto PackageIt = AllPackageNamesToMove.CreateConstIterator(); PackageIt; ++PackageIt )
 		{
-			ReportPackageNames.Add((*PackageIt).ToString());
+			ReportPackages.Get()->Add({ (*PackageIt).ToString(), true });
 		}
-		SPackageReportDialog::FOnReportConfirmed OnReportConfirmed = SPackageReportDialog::FOnReportConfirmed::CreateUObject(this, &UAssetToolsImpl::MigratePackages_ReportConfirmed, ReportPackageNames);
-		SPackageReportDialog::OpenPackageReportDialog(ReportMessage, ReportPackageNames, OnReportConfirmed);
+		SPackageReportDialog::FOnReportConfirmed OnReportConfirmed = SPackageReportDialog::FOnReportConfirmed::CreateUObject(this, &UAssetToolsImpl::MigratePackages_ReportConfirmed, ReportPackages);
+		SPackageReportDialog::OpenPackageReportDialog(ReportMessage, *ReportPackages.Get(), OnReportConfirmed);
 	}
 }
 
-void UAssetToolsImpl::MigratePackages_ReportConfirmed(TArray<FString> ConfirmedPackageNamesToMigrate) const
+void UAssetToolsImpl::MigratePackages_ReportConfirmed(TSharedPtr<TArray<ReportPackageData>> PackageDataToMigrate) const
 {
 	// Choose a destination folder
 	IDesktopPlatform* DesktopPlatform = FDesktopPlatformModule::Get();
@@ -2636,14 +2636,19 @@ void UAssetToolsImpl::MigratePackages_ReportConfirmed(TArray<FString> ConfirmedP
 
 	SlowTask.EnterProgressFrame();
 	{
-		FScopedSlowTask LoopProgress(ConfirmedPackageNamesToMigrate.Num());
-		for ( auto PackageNameIt = ConfirmedPackageNamesToMigrate.CreateConstIterator(); PackageNameIt; ++PackageNameIt )
+		FScopedSlowTask LoopProgress(PackageDataToMigrate.Get()->Num());
+		for ( auto PackageDataIt = PackageDataToMigrate.Get()->CreateConstIterator(); PackageDataIt; ++PackageDataIt)
 		{
 			LoopProgress.EnterProgressFrame();
+			if (!PackageDataIt->bShouldMigratePackage)
+			{
+				continue;
+			}
 
-			const FString& PackageName = *PackageNameIt;
+			const FString& PackageName = PackageDataIt->Name;
 			FString SrcFilename;
-			if ( !FPackageName::DoesPackageExist(PackageName, nullptr, &SrcFilename) )
+			
+			if (!FPackageName::DoesPackageExist(PackageName, nullptr, &SrcFilename))
 			{
 				const FText ErrorMessage = FText::Format(LOCTEXT("MigratePackages_PackageMissing", "{0} does not exist on disk."), FText::FromString(PackageName));
 				UE_LOG(LogAssetTools, Warning, TEXT("%s"), *ErrorMessage.ToString());
diff --git a/Engine/Source/Developer/AssetTools/Private/AssetTools.h b/Engine/Source/Developer/AssetTools/Private/AssetTools.h
--- a/Engine/Source/Developer/AssetTools/Private/AssetTools.h
+++ b/Engine/Source/Developer/AssetTools/Private/AssetTools.h
@@ -15,6 +15,7 @@ class IClassTypeActions;
 class UAutomatedAssetImportData;
 class UFactory;
 class UAssetImportTask;
+struct ReportPackageData;
 
 /** Parameters for importing specific set of files */
 struct FAssetImportParams
@@ -51,7 +52,7 @@ PRAGMA_DISABLE_DEPRECATION_WARNINGS
 UCLASS(transient)
 class UAssetToolsImpl : public UObject, public IAssetTools
 {
-	GENERATED_BODY()
+	GENERATED_BODY() 
 public:
 	UAssetToolsImpl(const FObjectInitializer& ObjectInitializer);
 
@@ -140,7 +141,7 @@ private:
 	void PerformAdvancedCopyPackages(TArray<FName> SelectedPackageNames, FString TargetPath) const;
 
 	/** Copies files after the final list was confirmed */
-	void MigratePackages_ReportConfirmed(TArray<FString> ConfirmedPackageNamesToMigrate) const;
+	void MigratePackages_ReportConfirmed(TSharedPtr<TArray<ReportPackageData>> PackageDataToMigrate) const;
 
 	/** Copies files after the final list was confirmed */
 	void AdvancedCopyPackages_ReportConfirmed(FAdvancedCopyParams CopyParam, TArray<TMap<FString, FString>> DestinationMap) const;
diff --git a/Engine/Source/Developer/AssetTools/Private/SPackageReportDialog.cpp b/Engine/Source/Developer/AssetTools/Private/SPackageReportDialog.cpp
--- a/Engine/Source/Developer/AssetTools/Private/SPackageReportDialog.cpp
+++ b/Engine/Source/Developer/AssetTools/Private/SPackageReportDialog.cpp
@@ -9,6 +9,7 @@
 #include "Framework/Application/SlateApplication.h"
 #include "Widgets/Images/SImage.h"
 #include "Widgets/Layout/SUniformGridPanel.h"
+#include "Widgets/Input/SCheckBox.h"
 #include "Widgets/Input/SButton.h"
 #include "EditorStyleSet.h"
 #include "Interfaces/IMainFrameModule.h"
@@ -24,20 +25,28 @@ struct FCompareFPackageReportNodeByName
 };
 
 FPackageReportNode::FPackageReportNode()
-	: bIsFolder(false)
+	: bIsChecked(true)
+	, bIsActive(true)
+	, bShouldMigratePackage(nullptr)
+	, bIsFolder(false)
+	, Parent(nullptr)
 {}
 
 FPackageReportNode::FPackageReportNode(const FString& InNodeName, bool InIsFolder)
 	: NodeName(InNodeName)
+	, bIsChecked(true)
+	, bIsActive(true)
+	, bShouldMigratePackage(nullptr)
 	, bIsFolder(InIsFolder)
+	, Parent(nullptr)
 {}
 
-void FPackageReportNode::AddPackage(const FString& PackageName)
+void FPackageReportNode::AddPackage(const FString& PackageName, bool* bInShouldMigratePackage)
 {
 	TArray<FString> PathElements;
 	PackageName.ParseIntoArray(PathElements, TEXT("/"), /*InCullEmpty=*/true);
 
-	return AddPackage_Recursive(PathElements);
+	return AddPackage_Recursive(PathElements, bInShouldMigratePackage);
 }
 
 void FPackageReportNode::ExpandChildrenRecursively(const TSharedRef<PackageReportTree>& Treeview)
@@ -49,7 +58,7 @@ void FPackageReportNode::ExpandChildrenRecursively(const TSharedRef<PackageRepor
 	}
 }
 
-void FPackageReportNode::AddPackage_Recursive(TArray<FString>& PathElements)
+void FPackageReportNode::AddPackage_Recursive(TArray<FString>& PathElements, bool* bInShouldMigratePackage)
 {
 	if ( PathElements.Num() > 0 )
 	{
@@ -74,18 +83,23 @@ void FPackageReportNode::AddPackage_Recursive(TArray<FString>& PathElements)
 			const bool bIsAFolder = (PathElements.Num() > 0);
 			int32 ChildIdx = Children.Add( MakeShareable(new FPackageReportNode(ChildNodeName, bIsAFolder)) );
 			Child = Children[ChildIdx];
+			Child.Get()->Parent = this;
 			Children.Sort( FCompareFPackageReportNodeByName() );
 		}
 
 		if ( ensure(Child.IsValid()) )
 		{
-			Child->AddPackage_Recursive(PathElements);
+			Child->AddPackage_Recursive(PathElements, bInShouldMigratePackage);
 		}
 	}
+	else 
+	{
+		bShouldMigratePackage = bInShouldMigratePackage;
+	}
 }
 
 BEGIN_SLATE_FUNCTION_BUILD_OPTIMIZATION
-void SPackageReportDialog::Construct( const FArguments& InArgs, const FText& InReportMessage, const TArray<FString>& InPackageNames, const FOnReportConfirmed& InOnReportConfirmed )
+void SPackageReportDialog::Construct( const FArguments& InArgs, const FText& InReportMessage, TArray<ReportPackageData>& InPackageNames, const FOnReportConfirmed& InOnReportConfirmed )
 {
 	OnReportConfirmed = InOnReportConfirmed;
 	FolderOpenBrush = FEditorStyle::GetBrush("ContentBrowser.AssetTreeFolderOpen");
@@ -165,7 +179,7 @@ void SPackageReportDialog::Construct( const FArguments& InArgs, const FText& InR
 }
 END_SLATE_FUNCTION_BUILD_OPTIMIZATION
 
-void SPackageReportDialog::OpenPackageReportDialog(const FText& ReportMessage, const TArray<FString>& PackageNames, const FOnReportConfirmed& InOnReportConfirmed)
+void SPackageReportDialog::OpenPackageReportDialog(const FText& ReportMessage, TArray<ReportPackageData>& PackageNames, const FOnReportConfirmed& InOnReportConfirmed)
 {
 	TSharedRef<SWindow> ReportWindow = SNew(SWindow)
 		.Title(LOCTEXT("ReportWindowTitle", "Asset Report"))
@@ -175,7 +189,7 @@ void SPackageReportDialog::OpenPackageReportDialog(const FText& ReportMessage, c
 		[
 			SNew(SPackageReportDialog, ReportMessage, PackageNames, InOnReportConfirmed)
 		];
-
+		
 	IMainFrameModule& MainFrameModule = FModuleManager::LoadModuleChecked<IMainFrameModule>(TEXT("MainFrame"));
 	if ( MainFrameModule.GetParentWindow().IsValid() )
 	{
@@ -197,7 +211,7 @@ void SPackageReportDialog::CloseDialog()
 	}
 }
 
-TSharedRef<ITableRow> SPackageReportDialog::GenerateTreeRow( TSharedPtr<FPackageReportNode> TreeItem, const TSharedRef<STableViewBase>& OwnerTable )
+TSharedRef<ITableRow> SPackageReportDialog::GenerateTreeRow( TSharedPtr<FPackageReportNode> TreeItem, const TSharedRef<STableViewBase>& OwnerTable)
 {
 	check(TreeItem.IsValid());
 
@@ -210,28 +224,69 @@ TSharedRef<ITableRow> SPackageReportDialog::GenerateTreeRow( TSharedPtr<FPackage
 			+SHorizontalBox::Slot()
 			.AutoWidth()
 			[
-				SNew(SImage).Image( IconBrush )
-			]
+				SNew(SCheckBox)
+				.OnCheckStateChanged(this, &SPackageReportDialog::CheckBoxStateChanged, TreeItem, OwnerTable)
+				.IsChecked(this, &SPackageReportDialog::GetEnabledCheckState, TreeItem)
+				.IsEnabled(TreeItem.Get()->Parent == nullptr ? true : TreeItem.Get()->Parent->bIsActive)
 
+			]
+			+SHorizontalBox::Slot()
+			.AutoWidth()
+			[
+				SNew(SImage).Image(IconBrush)
+			]
 			// Name
 			+SHorizontalBox::Slot()
 			.FillWidth(1.f)
 			[
 				SNew(STextBlock).Text(FText::FromString(TreeItem->NodeName))
+				.ColorAndOpacity(TreeItem.Get()->bIsActive ? FSlateColor::UseForeground() : FSlateColor::UseSubduedForeground())
 			]
 		];
 }
 
+ECheckBoxState SPackageReportDialog::GetEnabledCheckState(TSharedPtr<FPackageReportNode> TreeItem) const
+{
+	return TreeItem.Get()->bIsChecked ? ECheckBoxState::Checked : ECheckBoxState::Unchecked;
+}
+
+void SPackageReportDialog::SetStateRecursive(TSharedPtr<FPackageReportNode> TreeItem, bool bWasChecked)
+{
+	if (TreeItem.Get() == nullptr) return;
+
+	TreeItem.Get()->bIsActive = bWasChecked && TreeItem.Get()->bIsChecked;
+	if (TreeItem.Get()->bShouldMigratePackage)
+	{
+		*(TreeItem.Get()->bShouldMigratePackage) = TreeItem.Get()->bIsActive;
+	}
+
+	TArray< TSharedPtr<FPackageReportNode> > Children;
+	GetChildrenForTree(TreeItem, Children);
+	for (int i = 0; i < Children.Num(); i++)
+	{
+		if (Children[i].Get() == nullptr) continue;
+
+		SetStateRecursive(Children[i], TreeItem.Get()->bIsActive);
+	}
+}
+
+void SPackageReportDialog::CheckBoxStateChanged(ECheckBoxState InCheckBoxState, TSharedPtr<FPackageReportNode> TreeItem, TSharedRef<STableViewBase> OwnerTable)
+{
+	TreeItem.Get()->bIsChecked = InCheckBoxState == ECheckBoxState::Checked;
+	SetStateRecursive(TreeItem, InCheckBoxState == ECheckBoxState::Checked);
+	OwnerTable.Get().RebuildList();
+}
+
 void SPackageReportDialog::GetChildrenForTree( TSharedPtr<FPackageReportNode> TreeItem, TArray< TSharedPtr<FPackageReportNode> >& OutChildren )
 {
 	OutChildren = TreeItem->Children;
 }
 
-void SPackageReportDialog::ConstructNodeTree(const TArray<FString>& PackageNames)
+void SPackageReportDialog::ConstructNodeTree(TArray<ReportPackageData>& PackageNames)
 {
-	for ( auto PackageIt = PackageNames.CreateConstIterator(); PackageIt; ++PackageIt )
+	for (ReportPackageData& Package : PackageNames)
 	{
-		PackageReportRootNode.AddPackage(*PackageIt);
+		PackageReportRootNode.AddPackage(Package.Name, &Package.bShouldMigratePackage);
 	}
 }
 
diff --git a/Engine/Source/Developer/AssetTools/Private/SPackageReportDialog.h b/Engine/Source/Developer/AssetTools/Private/SPackageReportDialog.h
--- a/Engine/Source/Developer/AssetTools/Private/SPackageReportDialog.h
+++ b/Engine/Source/Developer/AssetTools/Private/SPackageReportDialog.h
@@ -14,12 +14,28 @@ struct FPackageReportNode;
 
 typedef STreeView< TSharedPtr<struct FPackageReportNode> > PackageReportTree;
 
+struct ReportPackageData
+{
+	FString Name;
+	bool bShouldMigratePackage;
+};
+
 struct FPackageReportNode
 {
 	/** The name of the tree node without the path */
-	FString NodeName;
+	FString NodeName; 
+	/** A user-exposed flag determining whether the content of this node and its children should be migrated or not. */
+	bool bIsChecked;
+	/** A flag determining whether this node should be migrated or not. This node is active as long as bIsChecked is true and if all the parent nodes are also checked. */
+	bool bIsActive;
+	/** A pointer to an external bool describing whether this node ultimately should be migrated or not. Is only non-null for leaf nodes.*/
+	bool* bShouldMigratePackage;
 	/** If true, this node is a folder instead of a package */
 	bool bIsFolder;
+	
+	/** The parent of this node */
+	FPackageReportNode* Parent;
+
 	/** The children of this node */
 	TArray< TSharedPtr<FPackageReportNode> > Children;
 
@@ -28,14 +44,14 @@ struct FPackageReportNode
 	FPackageReportNode(const FString& InNodeName, bool InIsFolder);
 
 	/** Adds the path to the tree relative to this node, creating nodes as needed. */
-	void AddPackage(const FString& PackageName);
+	void AddPackage(const FString& PackageName, bool* bInIsPackageIncluded);
 
 	/** Expands this node and all its children */
 	void ExpandChildrenRecursively(const TSharedRef<PackageReportTree>& Treeview);
 
 private:
 	/** Helper function for AddPackage. PathElements is the tokenized path delimited by "/" */
-	void AddPackage_Recursive(TArray<FString>& PathElements);
+	void AddPackage_Recursive(TArray<FString>& PathElements, bool* bInIsPackageIncluded);
 };
 
 class SPackageReportDialog : public SCompoundWidget
@@ -48,17 +64,26 @@ public:
 	SLATE_END_ARGS()
 
 	/** Constructs this widget with InArgs */
-	void Construct( const FArguments& InArgs, const FText& InReportMessage, const TArray<FString>& InPackageNames, const FOnReportConfirmed& InOnReportConfirmed );
+	void Construct( const FArguments& InArgs, const FText& InReportMessage, TArray<ReportPackageData>& InPackageNames, const FOnReportConfirmed& InOnReportConfirmed );
 
 	/** Opens the dialog in a new window */
-	static void OpenPackageReportDialog(const FText& ReportMessage, const TArray<FString>& PackageNames, const FOnReportConfirmed& InOnReportConfirmed);
+	static void OpenPackageReportDialog(const FText& ReportMessage, TArray<ReportPackageData>& PackageNames, const FOnReportConfirmed& InOnReportConfirmed);
 
 	/** Closes the dialog. */
 	void CloseDialog();
 
 private:
+	/** Recursively sets the checked/active state of every child of this node in the tree when a checkbox is toggled. */
+	void SetStateRecursive(TSharedPtr<FPackageReportNode> TreeItem, bool bWasChecked);
+
+	/** Callback to check whether a checkbox is checked or not. */
+	ECheckBoxState GetEnabledCheckState(TSharedPtr<FPackageReportNode> TreeItem) const;
+
+	/** Callback called whenever a checkbox is toggled. */
+	void CheckBoxStateChanged(ECheckBoxState InCheckBoxState, TSharedPtr<FPackageReportNode> TreeItem, TSharedRef<STableViewBase> OwnerTable);
+
 	/** Constructs the node tree given the list of package names */
-	void ConstructNodeTree(const TArray<FString>& PackageNames);
+	void ConstructNodeTree(TArray<ReportPackageData>& PackageNames);
 
 	/** Handler to generate a row in the report tree */
 	TSharedRef<ITableRow> GenerateTreeRow( TSharedPtr<FPackageReportNode> TreeItem, const TSharedRef<STableViewBase>& OwnerTable );
