Fix for AdvancedCopy speeds being very slow.
#rnx

diff --git a/Engine/Source/Developer/AssetTools/Private/AssetTools.cpp b/Engine/Source/Developer/AssetTools/Private/AssetTools.cpp
--- a/Engine/Source/Developer/AssetTools/Private/AssetTools.cpp
+++ b/Engine/Source/Developer/AssetTools/Private/AssetTools.cpp
@@ -800,7 +800,7 @@ void UAssetToolsImpl::GetAllAdvancedCopySources(FName SelectedPackage, FAdvanced
 		// Folders should ALWAYS get checked for assets and subfolders
 		if ((CopyParams.bShouldCheckForDependencies && SourceAssetData.Num() > 0) || bIsFolder)
 		{
-			RecursiveGetDependenciesAdvanced(SelectedPackage, CopyParams, CurrentDependencies, DependencyMap, CopyCustomization);
+			RecursiveGetDependenciesAdvanced(SelectedPackage, CopyParams, CurrentDependencies, DependencyMap, CopyCustomization, SourceAssetData);
 		}
 		OutPackageNamesToCopy.Append(CurrentDependencies);
 	}
@@ -2823,31 +2823,31 @@ void UAssetToolsImpl::RecursiveGetDependencies(const FName& PackageName, TSet<FN
 	}
 }
 
-void UAssetToolsImpl::RecursiveGetDependenciesAdvanced(const FName& PackageName, FAdvancedCopyParams& CopyParams, TArray<FName>& AllDependencies, TMap<FName, FName>& DependencyMap, const UAdvancedCopyCustomization* CopyCustomization) const
+void UAssetToolsImpl::RecursiveGetDependenciesAdvanced(const FName& PackageName, FAdvancedCopyParams& CopyParams, TArray<FName>& AllDependencies, TMap<FName, FName>& DependencyMap, const UAdvancedCopyCustomization* CopyCustomization, TArray<FAssetData>& OptionalAssetData) const
 {
 	FAssetRegistryModule& AssetRegistryModule = FModuleManager::Get().LoadModuleChecked<FAssetRegistryModule>(TEXT("AssetRegistry"));
 	TArray<FName> Dependencies;
-	TArray<FAssetData> PackageAssetData;
 	IAssetRegistry& AssetRegistry = AssetRegistryModule.Get();
-	AssetRegistry.GetAssetsByPackageName(PackageName, PackageAssetData);
 	// We found an asset
-	if (PackageAssetData.Num() > 0)
+	if (OptionalAssetData.Num() > 0)
 	{
 		AssetRegistryModule.Get().GetDependencies(PackageName, Dependencies);
 		for (auto DependsIt = Dependencies.CreateConstIterator(); DependsIt; ++DependsIt)
 		{
 			if (!AllDependencies.Contains(*DependsIt))
 			{
-				TArray<FAssetData> DependencyAssetData;
-				if (AssetRegistry.GetAssetsByPackageName(*DependsIt, DependencyAssetData))
+				FAssetData DependencyAsset = AssetRegistry.GetAssetByObjectPath(*DependsIt);
+				if (DependencyAsset != FAssetData())
 				{
+					TArray<FAssetData> DependencyAssetData;
+					DependencyAssetData.Add(DependencyAsset);
 					FARFilter ExclusionFilter = CopyCustomization->GetARFilter();
 					AssetRegistry.UseFilterToExcludeAssets(DependencyAssetData, ExclusionFilter);
 					if (DependencyAssetData.IsValidIndex(0))
 					{
 						AllDependencies.Add(*DependsIt);
 						DependencyMap.Add(*DependsIt, PackageName);
-						RecursiveGetDependenciesAdvanced(*DependsIt, CopyParams, AllDependencies, DependencyMap, CopyCustomization);
+						RecursiveGetDependenciesAdvanced(*DependsIt, CopyParams, AllDependencies, DependencyMap, CopyCustomization, DependencyAssetData);
 					}
 				}
 
@@ -2868,7 +2868,7 @@ void UAssetToolsImpl::RecursiveGetDependenciesAdvanced(const FName& PackageName,
 				// If we should check the assets we found for dependencies
 				if (CopyParams.bShouldCheckForDependencies)
 				{
-					RecursiveGetDependenciesAdvanced(FName(*Asset.GetPackage()->GetName()), CopyParams, AllDependencies, DependencyMap, CopyCustomization);
+					RecursiveGetDependenciesAdvanced(FName(*Asset.GetPackage()->GetName()), CopyParams, AllDependencies, DependencyMap, CopyCustomization, PathAssetData);
 				}
 			}
 		}
@@ -2879,7 +2879,8 @@ void UAssetToolsImpl::RecursiveGetDependenciesAdvanced(const FName& PackageName,
 			AssetRegistry.GetSubPaths(PackageName.ToString(), SubPaths, false);
 			for (const FString SubPath : SubPaths)
 			{
-				RecursiveGetDependenciesAdvanced(FName(*SubPath), CopyParams, AllDependencies, DependencyMap, CopyCustomization);
+				TArray<FAssetData> EmptyArray;
+				RecursiveGetDependenciesAdvanced(FName(*SubPath), CopyParams, AllDependencies, DependencyMap, CopyCustomization, EmptyArray);
 			}
 
 		}
diff --git a/Engine/Source/Developer/AssetTools/Private/AssetTools.h b/Engine/Source/Developer/AssetTools/Private/AssetTools.h
--- a/Engine/Source/Developer/AssetTools/Private/AssetTools.h
+++ b/Engine/Source/Developer/AssetTools/Private/AssetTools.h
@@ -150,7 +150,7 @@ private:
 	void RecursiveGetDependencies(const FName& PackageName, TSet<FName>& AllDependencies) const;
 
 	/** Gets the dependencies of the specified package recursively while omitting things that don't pass the FARFilter passed in from FAdvancedCopyParams */
-	void RecursiveGetDependenciesAdvanced(const FName& PackageName, FAdvancedCopyParams& CopyParams, TArray<FName>& AllDependencies, TMap<FName, FName>& DependencyMap, const class UAdvancedCopyCustomization* CopyCustomization) const;
+	void RecursiveGetDependenciesAdvanced(const FName& PackageName, FAdvancedCopyParams& CopyParams, TArray<FName>& AllDependencies, TMap<FName, FName>& DependencyMap, const class UAdvancedCopyCustomization* CopyCustomization, TArray<FAssetData>& OptionalAssetData) const;
 
 	/** Records the time taken for an import and reports it to engine analytics, if available */
 	static void OnNewImportRecord(UClass* AssetType, const FString& FileExtension, bool bSucceeded, bool bWasCancelled, const FDateTime& StartTime);
