Fix dereference of nullptr when reporting error during VT asset conversion

diff --git a/Engine/Source/Developer/AssetTools/Private/AssetVtConversion.cpp b/Engine/Source/Developer/AssetTools/Private/AssetVtConversion.cpp
--- a/Engine/Source/Developer/AssetTools/Private/AssetVtConversion.cpp
+++ b/Engine/Source/Developer/AssetTools/Private/AssetVtConversion.cpp
@@ -530,7 +530,7 @@ void FVTConversionWorker::DoConvert(bool backwards)
 				EngineTextureToCopyMap.Add(Tex, TextureToUpdate);
 				if (!TextureToUpdate)
 				{
-					UE_LOG(LogVirtualTextureConversion, Warning, TEXT("Failed to duplicate engine texture %s"), *TextureToUpdate->GetPathName());
+					UE_LOG(LogVirtualTextureConversion, Warning, TEXT("Failed to duplicate engine texture %s"), *Tex->GetPathName());
 				}
 			}
 
