Only show convert to VT/non-VT context menu items if they are applicable for the selected texture(s)

diff --git a/Engine/Source/Developer/AssetTools/Private/AssetTypeActions/AssetTypeActions_Texture.cpp b/Engine/Source/Developer/AssetTools/Private/AssetTypeActions/AssetTypeActions_Texture.cpp
--- a/Engine/Source/Developer/AssetTools/Private/AssetTypeActions/AssetTypeActions_Texture.cpp
+++ b/Engine/Source/Developer/AssetTools/Private/AssetTypeActions/AssetTypeActions_Texture.cpp
@@ -47,6 +47,19 @@
 void FAssetTypeActions_Texture::GetActions( const TArray<UObject*>& InObjects, FMenuBuilder& MenuBuilder )
 {
 	auto Textures = GetTypedWeakObjectPtrs<UTexture>(InObjects);
+	bool bHasVirtualTextures = false;
+	bool bHasNonVirtualTextures = false;
+	for (const TWeakObjectPtr<UTexture>& Texture : Textures)
+	{
+		if (Texture->VirtualTextureStreaming)
+		{
+			bHasVirtualTextures = true;
+		}
+		else
+		{
+			bHasNonVirtualTextures = true;
+		}
+	}
 
 	MenuBuilder.AddMenuEntry(
 		LOCTEXT("Texture_CreateMaterial", "Create Material"),
@@ -58,25 +71,31 @@ void FAssetTypeActions_Texture::GetActions( const TArray<UObject*>& InObjects, F
 			)
 		);
 
-	MenuBuilder.AddMenuEntry(
-		LOCTEXT("Texture_ConvertToVT", "Convert to Virtual Texture"),
-		LOCTEXT("Texture_ConvertToVTTooltip", "Converts this texture to a virtual texture if it fit's the size limit imposed in the texture importer settings."),
-		FSlateIcon(FEditorStyle::GetStyleSetName(), "ClassIcon.Texture2D"),
-		FUIAction(
-			FExecuteAction::CreateSP(this, &FAssetTypeActions_Texture::ExecuteConvertToVirtualTexture, Textures),
-			FCanExecuteAction()
-		)
-	);
+	if (bHasNonVirtualTextures)
+	{
+		MenuBuilder.AddMenuEntry(
+			LOCTEXT("Texture_ConvertToVT", "Convert to Virtual Texture"),
+			LOCTEXT("Texture_ConvertToVTTooltip", "Converts this texture to a virtual texture if it fit's the size limit imposed in the texture importer settings."),
+			FSlateIcon(FEditorStyle::GetStyleSetName(), "ClassIcon.Texture2D"),
+			FUIAction(
+				FExecuteAction::CreateSP(this, &FAssetTypeActions_Texture::ExecuteConvertToVirtualTexture, Textures),
+				FCanExecuteAction()
+			)
+		);
+	}
 
-	MenuBuilder.AddMenuEntry(
-		LOCTEXT("Texture_ConvertToRegular", "Convert to Regular Texture"),
-		LOCTEXT("Texture_ConvertToRegularTooltip", "Converts this texture to a ragular 2D texture if it is a virtual texture."),
-		FSlateIcon(FEditorStyle::GetStyleSetName(), "ClassIcon.Texture2D"),
-		FUIAction(
-			FExecuteAction::CreateSP(this, &FAssetTypeActions_Texture::ExecuteConvertToRegularTexture, Textures),
-			FCanExecuteAction()
-		)
-	);
+	if (bHasVirtualTextures)
+	{
+		MenuBuilder.AddMenuEntry(
+			LOCTEXT("Texture_ConvertToRegular", "Convert to Regular Texture"),
+			LOCTEXT("Texture_ConvertToRegularTooltip", "Converts this texture to a ragular 2D texture if it is a virtual texture."),
+			FSlateIcon(FEditorStyle::GetStyleSetName(), "ClassIcon.Texture2D"),
+			FUIAction(
+				FExecuteAction::CreateSP(this, &FAssetTypeActions_Texture::ExecuteConvertToRegularTexture, Textures),
+				FCanExecuteAction()
+			)
+		);
+	}
 
 	if ( InObjects.Num() == 1 )
 	{
