Don't display option to convert non-2D textures to VT
Fix some UI typos

diff --git a/Engine/Source/Developer/AssetTools/Private/AssetTypeActions/AssetTypeActions_Texture.cpp b/Engine/Source/Developer/AssetTools/Private/AssetTypeActions/AssetTypeActions_Texture.cpp
--- a/Engine/Source/Developer/AssetTools/Private/AssetTypeActions/AssetTypeActions_Texture.cpp
+++ b/Engine/Source/Developer/AssetTools/Private/AssetTypeActions/AssetTypeActions_Texture.cpp
@@ -55,8 +55,9 @@ void FAssetTypeActions_Texture::GetActions( const TArray<UObject*>& InObjects, F
 		{
 			bHasVirtualTextures = true;
 		}
-		else
+		else if(Cast<UTexture2D>(Texture.Get()))
 		{
+			// Currently only texture2D may be converted to VT
 			bHasNonVirtualTextures = true;
 		}
 	}
@@ -75,7 +76,7 @@ void FAssetTypeActions_Texture::GetActions( const TArray<UObject*>& InObjects, F
 	{
 		MenuBuilder.AddMenuEntry(
 			LOCTEXT("Texture_ConvertToVT", "Convert to Virtual Texture"),
-			LOCTEXT("Texture_ConvertToVTTooltip", "Converts this texture to a virtual texture if it fit's the size limit imposed in the texture importer settings."),
+			LOCTEXT("Texture_ConvertToVTTooltip", "Converts this texture to a virtual texture if it fits the size limit imposed in the texture importer settings."),
 			FSlateIcon(FEditorStyle::GetStyleSetName(), "ClassIcon.Texture2D"),
 			FUIAction(
 				FExecuteAction::CreateSP(this, &FAssetTypeActions_Texture::ExecuteConvertToVirtualTexture, Textures),
@@ -88,7 +89,7 @@ void FAssetTypeActions_Texture::GetActions( const TArray<UObject*>& InObjects, F
 	{
 		MenuBuilder.AddMenuEntry(
 			LOCTEXT("Texture_ConvertToRegular", "Convert to Regular Texture"),
-			LOCTEXT("Texture_ConvertToRegularTooltip", "Converts this texture to a ragular 2D texture if it is a virtual texture."),
+			LOCTEXT("Texture_ConvertToRegularTooltip", "Converts this texture to a regular 2D texture if it is a virtual texture."),
 			FSlateIcon(FEditorStyle::GetStyleSetName(), "ClassIcon.Texture2D"),
 			FUIAction(
 				FExecuteAction::CreateSP(this, &FAssetTypeActions_Texture::ExecuteConvertToRegularTexture, Textures),
