VT conversion tool also makes updates to any preview materials using selected textures, this ensures any active material editors remain valid after conversion

diff --git a/Engine/Source/Developer/AssetTools/Private/AssetTypeActions/AssetTypeActions_Texture.cpp b/Engine/Source/Developer/AssetTools/Private/AssetTypeActions/AssetTypeActions_Texture.cpp
--- a/Engine/Source/Developer/AssetTools/Private/AssetTypeActions/AssetTypeActions_Texture.cpp
+++ b/Engine/Source/Developer/AssetTools/Private/AssetTypeActions/AssetTypeActions_Texture.cpp
@@ -659,14 +659,22 @@ private:
 
 		for (UMaterial *RootMat : Worker.Materials)
 		{
-			AssetList.Add(RootMat);
-			AssetStatus.Add(FConversionStatus());
+			// Don't want to show anything from the transient package, this may include preview materials from any currently active material editor
+			// We patch these up so material editor remains valid, but not useful to display this to user
+			if (RootMat->GetOutermost() != GetTransientPackage())
+			{
+				AssetList.Add(RootMat);
+				AssetStatus.Add(FConversionStatus());
+			}
 		}
 
 		for (UMaterialFunctionInterface *Func : Worker.Functions)
 		{
-			AssetList.Add(Func);
-			AssetStatus.Add(FConversionStatus());
+			if (Func->GetOutermost() != GetTransientPackage())
+			{
+				AssetList.Add(Func);
+				AssetStatus.Add(FConversionStatus());
+			}
 		}
 
 		check(AssetList.Num() == AssetStatus.Num());
diff --git a/Engine/Source/Developer/AssetTools/Private/AssetVtConversion.cpp b/Engine/Source/Developer/AssetTools/Private/AssetVtConversion.cpp
--- a/Engine/Source/Developer/AssetTools/Private/AssetVtConversion.cpp
+++ b/Engine/Source/Developer/AssetTools/Private/AssetVtConversion.cpp
@@ -6,6 +6,7 @@
 #include "Materials/Material.h"
 #include "Materials/MaterialInstance.h"
 #include "Materials/MaterialFunctionInstance.h"
+#include "MaterialEditor/PreviewMaterial.h"
 #include "Factories/MaterialFactoryNew.h"
 #include "EditorFramework/AssetImportData.h"
 #include "AssetTools.h"
@@ -83,6 +84,23 @@ template <class T> void GetReferencersOfType(UObject *Object, TArray<T*> &OutObj
 	}
 }
 
+static void GetPreviewMaterials(UTexture2D* InTexture, TArray<UMaterialInterface*>& OutMaterials)
+{
+	for (TObjectIterator<UPreviewMaterial> It; It; ++It)
+	{
+		UMaterial* Material = *It;
+		for (uint32 PropertyIndex = 0u; PropertyIndex < MP_MAX; ++PropertyIndex)
+		{
+			const EMaterialProperty PropertyToValidate = (EMaterialProperty)PropertyIndex;
+			if (Material->IsTextureReferencedByProperty(PropertyToValidate, InTexture))
+			{
+				OutMaterials.Add(Material);
+				break;
+			}
+		}
+	}
+}
+
 bool IsVirtualTextureValidForMaterial(UMaterialInterface* InMaterial, UTexture2D* InTexture)
 {
 	for (uint32 PropertyIndex = 0u; PropertyIndex < MP_MAX; ++PropertyIndex)
@@ -155,6 +173,10 @@ void FVTConversionWorker::FindAllTexturesAndMaterials_Iteration(TArray<UMaterial
 			// since we're not changing the original texture, we don't need to bring in any additional dependencies
 			if (!Tex2D->GetPathName().StartsWith("/Engine/"))
 			{
+				// Also get any preview materials that reference the given texture
+				// We need to convert these to ensure any active material editors remain valid
+				GetPreviewMaterials(Tex2D, MaterialsUsedByAffectedTextures);
+
 				MaterialsUsedByAffectedTextures.Append(MaterialsUsingTexture);
 			}
 			++TextureIndex;
