Message user on attempt to convert non-power-2 texture to VT

diff --git a/Engine/Source/Developer/AssetTools/Private/AssetTypeActions/AssetTypeActions_Texture.cpp b/Engine/Source/Developer/AssetTools/Private/AssetTypeActions/AssetTypeActions_Texture.cpp
--- a/Engine/Source/Developer/AssetTools/Private/AssetTypeActions/AssetTypeActions_Texture.cpp
+++ b/Engine/Source/Developer/AssetTools/Private/AssetTypeActions/AssetTypeActions_Texture.cpp
@@ -197,6 +197,7 @@ struct FConversionStatus
 {
 	bool UserSelected = true; // Originally selected by the user
 	bool UnderSized = false; // Too small to convert according to user settings
+	bool NonPowerOf2 = false; // Not a power of 2, can't convert
 	bool InvalidMaterialUsage = false;
 };
 
@@ -427,6 +428,11 @@ private:
 				SeverityIcon = "MessageLog.Error";
 				DetailedInfoText = LOCTEXT("ConvertToVT_ToolTip_InvalidUsage", "The texture could not be converted to VT due to its usage in materials (may be connected to a property that doesn't support VT).");
 			}
+			else if (Status.NonPowerOf2)
+			{
+				SeverityIcon = "MessageLog.Error";
+				DetailedInfoText = LOCTEXT("ConvertToVT_ToolTip_NonPowerOf2", "The texture could not be converted to VT because its size is not a power of 2.");
+			}
 			else if (bEngineAsset)
 			{
 				SeverityIcon = "MessageLog.Note";
@@ -445,6 +451,11 @@ private:
 				SeverityIcon = "MessageLog.Warning";
 				DetailedInfoText = LOCTEXT("ConvertToVT_ToolTip_InvalidUsageNotSelected", "The texture could not be converted to VT due to its usage in materials (may be connected to a property that doesn't support VT).");
 			}
+			else if (Status.NonPowerOf2)
+			{
+				SeverityIcon = "MessageLog.Warning";
+				DetailedInfoText = LOCTEXT("ConvertToVT_ToolTip_NonPowerOf2NotSelected", "The texture could not be converted to VT because its size is not a power of 2.");
+			}
 			else if (bEngineAsset)
 			{
 				SeverityIcon = "MessageLog.Note";
@@ -646,6 +657,7 @@ private:
 			FConversionStatus* Status = new(AssetStatus) FConversionStatus();
 			Status->UserSelected = Worker.UserTextures.Contains(Texture);
 			Status->UnderSized = (Texture->GetSizeX()*Texture->GetSizeY() < ThresholdValue*ThresholdValue);
+			Status->NonPowerOf2 = !Texture->Source.IsPowerOfTwo() && Texture->PowerOfTwoMode == ETexturePowerOfTwoSetting::None;
 		}
 
 		for (UTexture2D *Texture : Worker.MaterialRejectedTextures)
@@ -654,6 +666,7 @@ private:
 			FConversionStatus* Status = new(AssetStatus) FConversionStatus();
 			Status->UserSelected = Worker.UserTextures.Contains(Texture);
 			Status->UnderSized = (Texture->GetSizeX()*Texture->GetSizeY() < ThresholdValue*ThresholdValue);
+			Status->NonPowerOf2 = !Texture->Source.IsPowerOfTwo() && Texture->PowerOfTwoMode == ETexturePowerOfTwoSetting::None;
 			Status->InvalidMaterialUsage = true;
 		}
 
diff --git a/Engine/Source/Developer/AssetTools/Private/AssetVtConversion.cpp b/Engine/Source/Developer/AssetTools/Private/AssetVtConversion.cpp
--- a/Engine/Source/Developer/AssetTools/Private/AssetVtConversion.cpp
+++ b/Engine/Source/Developer/AssetTools/Private/AssetVtConversion.cpp
@@ -169,9 +169,11 @@ void FVTConversionWorker::FindAllTexturesAndMaterials_Iteration(TArray<UMaterial
 
 		if (bIsVirtualTextureValid)
 		{
-			// if this is an engine texture, we'll make a copy and update any needed references to point to the new copy
+			// If this is an engine texture, we'll make a copy and update any needed references to point to the new copy
 			// since we're not changing the original texture, we don't need to bring in any additional dependencies
-			if (!Tex2D->GetPathName().StartsWith("/Engine/"))
+			// Non-power-2 textures won't convert to VT, don't bring any dependencies for them
+			if (!Tex2D->GetPathName().StartsWith("/Engine/") &&
+				(Tex2D->Source.IsPowerOfTwo() || Tex2D->PowerOfTwoMode != ETexturePowerOfTwoSetting::None))
 			{
 				// Also get any preview materials that reference the given texture
 				// We need to convert these to ensure any active material editors remain valid
