Fix one more VS2019 compile errors
#rnx

diff --git a/Engine/Source/Developer/AssetTools/Private/AssetTypeActions/AssetTypeActions_Texture.cpp b/Engine/Source/Developer/AssetTools/Private/AssetTypeActions/AssetTypeActions_Texture.cpp
--- a/Engine/Source/Developer/AssetTools/Private/AssetTypeActions/AssetTypeActions_Texture.cpp
+++ b/Engine/Source/Developer/AssetTools/Private/AssetTypeActions/AssetTypeActions_Texture.cpp
@@ -348,7 +348,7 @@ public:
 						.HAlign(HAlign_Center)
 						.ContentPadding(FEditorStyle::GetMargin("StandardDialog.ContentPadding"))
 						.OnClicked(this, &SConvertToVTDlg::OnFilterButtonClicked)
-						.IsEnabled(&SConvertToVTDlg::GetFilterButtonEnabled)
+						.IsEnabled(this, &SConvertToVTDlg::GetFilterButtonEnabled)
 						.Text(LOCTEXT("ConvertToVT_Filter", "Apply Filter"))
 					]
 				]
@@ -375,7 +375,7 @@ public:
 						.HAlign(HAlign_Center)
 						.ContentPadding(FEditorStyle::GetMargin("StandardDialog.ContentPadding"))
 						.OnClicked(this, &SConvertToVTDlg::OnButtonClick, FConvertToVTDlg::Confirm)
-						.IsEnabled(&SConvertToVTDlg::GetOkButtonEnabled)
+						.IsEnabled(this, &SConvertToVTDlg::GetOkButtonEnabled)
 						.Text(LOCTEXT("ConvertToVT_OK", "OK"))
 					]
 					+ SUniformGridPanel::Slot(1, 0)
