DR - Fix missing copyright

diff --git a/Engine/Source/Developer/AssetTools/Private/AssetTypeActions/AssetTypeActions_Texture2DArray.cpp b/Engine/Source/Developer/AssetTools/Private/AssetTypeActions/AssetTypeActions_Texture2DArray.cpp
--- a/Engine/Source/Developer/AssetTools/Private/AssetTypeActions/AssetTypeActions_Texture2DArray.cpp
+++ b/Engine/Source/Developer/AssetTools/Private/AssetTypeActions/AssetTypeActions_Texture2DArray.cpp
@@ -1,2 +1,3 @@
-#include "AssetTypeActions_Texture2DArray.h"
+// Copyright 1998-2019 Epic Games, Inc. All Rights Reserved.
 
+#include "AssetTypeActions_Texture2DArray.h"
