Fixed compilation error.
#rnx

diff --git a/Engine/Source/Developer/AssetTools/Private/AssetTypeActions/AssetTypeActions_Texture2D.cpp b/Engine/Source/Developer/AssetTools/Private/AssetTypeActions/AssetTypeActions_Texture2D.cpp
--- a/Engine/Source/Developer/AssetTools/Private/AssetTypeActions/AssetTypeActions_Texture2D.cpp
+++ b/Engine/Source/Developer/AssetTools/Private/AssetTypeActions/AssetTypeActions_Texture2D.cpp
@@ -33,7 +33,8 @@ void FAssetTypeActions_Texture2D::GetActions(const TArray<UObject*>& InObjects,
 	static const auto AllowTextureArrayAssetCreationVar = IConsoleManager::Get().FindTConsoleVariableDataInt(TEXT("r.AllowTexture2DArrayCreation"));
 	if (Textures.Num() > 0 && AllowTextureArrayAssetCreationVar->GetValueOnGameThread() != 0)
 	{
-		MenuBuilder.AddMenuEntry(
+		Section.AddMenuEntry(
+			"Texture_Texture2DArray",
 			LOCTEXT("Texture_Texture2DArray", "Create Texture Array"),
 			LOCTEXT("Texture_CreateTexture2DArrayTooltip", "Creates a new texture array."),
 			FSlateIcon(FEditorStyle::GetStyleSetName(), "ClassIcon.Texture2D"),
