Fix crash when deleting actor from level

diff --git a/Engine/Source/Developer/AssetTools/Private/AssetRenameManager.cpp b/Engine/Source/Developer/AssetTools/Private/AssetRenameManager.cpp
--- a/Engine/Source/Developer/AssetTools/Private/AssetRenameManager.cpp
+++ b/Engine/Source/Developer/AssetTools/Private/AssetRenameManager.cpp
@@ -899,10 +899,20 @@ struct FSoftObjectPathRenameSerializer : public FArchiveUObject
 		, bSearchOnly(bInCheckOnly)
 		, bFoundReference(false)
 	{
+		if (InCachedObjectPaths)
+		{
+			DirtyDelegateHandle = UPackage::PackageMarkedDirtyEvent.AddRaw(this, &FSoftObjectPathRenameSerializer::OnMarkPackageDirty);
+		}
+
 		// Mark it as saving to correctly process all references
 		this->SetIsSaving(true);
 	}
 
+	virtual ~FSoftObjectPathRenameSerializer()
+	{
+		UPackage::PackageMarkedDirtyEvent.Remove(DirtyDelegateHandle);
+	}
+
 	virtual bool ShouldSkipProperty(const UProperty* InProperty) const override
 	{
 		if (InProperty->HasAnyPropertyFlags(CPF_Transient | CPF_Deprecated | CPF_IsPlainOldData))
@@ -997,9 +1007,22 @@ struct FSoftObjectPathRenameSerializer : public FArchiveUObject
 		return *this;
 	}
 
+	void OnMarkPackageDirty(UPackage* Pkg, bool bWasDirty)
+	{
+		UPackage::PackageMarkedDirtyEvent.Remove(DirtyDelegateHandle);
+
+		// CachedObjectPaths is no longer valid because an object was modified during serialization
+		if (CachedObjectPaths)
+		{
+			CachedObjectPaths = nullptr;
+			UE_LOG(LogAssetTools, VeryVerbose, TEXT("Performance: Package unexpectedly modified during serialization by FSoftObjectPathRenameSerializer: %s"), *Pkg->GetFullName());
+		}
+	}
+
 private:
 	const TMap<FSoftObjectPath, FSoftObjectPath>& RedirectorMap;
 	TMap<FSoftObjectPath, TSet<FWeakObjectPtr>>* CachedObjectPaths;
+	FDelegateHandle DirtyDelegateHandle;
 	UObject* CurrentObject;
 	bool bSearchOnly;
 	bool bFoundReference;
