Fix crash when deleting actor from level (follow up)

diff --git a/Engine/Source/Developer/AssetTools/Private/AssetRenameManager.cpp b/Engine/Source/Developer/AssetTools/Private/AssetRenameManager.cpp
--- a/Engine/Source/Developer/AssetTools/Private/AssetRenameManager.cpp
+++ b/Engine/Source/Developer/AssetTools/Private/AssetRenameManager.cpp
@@ -892,10 +892,11 @@ struct FSoftObjectPathRenameSerializer : public FArchiveUObject
 		return bFoundReference; 
 	}
 
-	FSoftObjectPathRenameSerializer(const TMap<FSoftObjectPath, FSoftObjectPath>& InRedirectorMap, bool bInCheckOnly, TMap<FSoftObjectPath, TSet<FWeakObjectPtr>>* InCachedObjectPaths)
+	FSoftObjectPathRenameSerializer(const TMap<FSoftObjectPath, FSoftObjectPath>& InRedirectorMap, bool bInCheckOnly, TMap<FSoftObjectPath, TSet<FWeakObjectPtr>>* InCachedObjectPaths, const FName InPackageName = NAME_None)
 		: RedirectorMap(InRedirectorMap)
 		, CachedObjectPaths(InCachedObjectPaths)
 		, CurrentObject(nullptr)
+		, PackageName(InPackageName)
 		, bSearchOnly(bInCheckOnly)
 		, bFoundReference(false)
 	{
@@ -1011,10 +1012,8 @@ struct FSoftObjectPathRenameSerializer : public FArchiveUObject
 	{
 		UPackage::PackageMarkedDirtyEvent.Remove(DirtyDelegateHandle);
 
-		// CachedObjectPaths is no longer valid because an object was modified during serialization
-		if (CachedObjectPaths)
+		if (CachedObjectPaths && Pkg && Pkg->GetFName() == PackageName)
 		{
-			CachedObjectPaths = nullptr;
 			UE_LOG(LogAssetTools, VeryVerbose, TEXT("Performance: Package unexpectedly modified during serialization by FSoftObjectPathRenameSerializer: %s"), *Pkg->GetFullName());
 		}
 	}
@@ -1024,6 +1023,7 @@ private:
 	TMap<FSoftObjectPath, TSet<FWeakObjectPtr>>* CachedObjectPaths;
 	FDelegateHandle DirtyDelegateHandle;
 	UObject* CurrentObject;
+	FName PackageName;
 	bool bSearchOnly;
 	bool bFoundReference;
 
@@ -1095,8 +1095,6 @@ bool FAssetRenameManager::CheckPackageForSoftObjectReferences(UPackage* Package,
 
 	if (CachedReferences == nullptr)
 	{
-		CachedReferences = &CachedSoftReferences.Add(Package->GetFName());
-
 		// Bind to dirty callback if we aren't already
 		if (!DirtyDelegateHandle.IsValid())
 		{
@@ -1105,7 +1103,8 @@ bool FAssetRenameManager::CheckPackageForSoftObjectReferences(UPackage* Package,
 
 		//Extract all objects soft references along with their referencer and cache them to avoid having to serialize again
 		TMap<FSoftObjectPath, FSoftObjectPath> EmptyMap;
-		FSoftObjectPathRenameSerializer CheckSerializer(EmptyMap, true, &CachedReferences->Map);
+		TMap<FSoftObjectPath, TSet<FWeakObjectPtr>> MapForCache;
+		FSoftObjectPathRenameSerializer CheckSerializer(EmptyMap, true, &MapForCache, Package->GetFName());
 
 		TArray<UObject*> ObjectsInPackage;
 		GetObjectsWithOuter(Package, ObjectsInPackage);
@@ -1121,6 +1120,9 @@ bool FAssetRenameManager::CheckPackageForSoftObjectReferences(UPackage* Package,
 			Object->Serialize(CheckSerializer);
 		}
 
+		CachedReferences = &CachedSoftReferences.Add(Package->GetFName());
+		CachedReferences->Map = MoveTemp(MapForCache);
+
 		CachedReferences->Map.GenerateKeyArray(CachedReferences->Keys);
 		
 		// Keys need to be sorted for binary search
