- Add IAssetTools::IsFixupReferencersInProgress
- Don't generate PPIDs upon world save if triggered by redirectors fixup.

diff --git a/Engine/Source/Developer/AssetTools/Private/AssetFixUpRedirectors.cpp b/Engine/Source/Developer/AssetTools/Private/AssetFixUpRedirectors.cpp
--- a/Engine/Source/Developer/AssetTools/Private/AssetFixUpRedirectors.cpp
+++ b/Engine/Source/Developer/AssetTools/Private/AssetFixUpRedirectors.cpp
@@ -72,6 +72,8 @@ void FAssetFixUpRedirectors::FixupReferencers(const TArray<UObjectRedirector*>&
 
 void FAssetFixUpRedirectors::ExecuteFixUp(TArray<TWeakObjectPtr<UObjectRedirector>> Objects) const
 {
+	TGuardValue<bool> Guard(bIsFixupReferencersInProgress, true);
+
 	TArray<FRedirectorRefs> RedirectorRefsList;
 	for (auto Object : Objects)
 	{
diff --git a/Engine/Source/Developer/AssetTools/Private/AssetFixUpRedirectors.h b/Engine/Source/Developer/AssetTools/Private/AssetFixUpRedirectors.h
--- a/Engine/Source/Developer/AssetTools/Private/AssetFixUpRedirectors.h
+++ b/Engine/Source/Developer/AssetTools/Private/AssetFixUpRedirectors.h
@@ -12,6 +12,9 @@ public:
 	/** Fix up references to the specified redirectors */
 	void FixupReferencers(const TArray<UObjectRedirector*>& Objects) const;
 
+	/** Returns whether redirectors are being fixed up. */
+	bool IsFixupReferencersInProgress() const { return bIsFixupReferencersInProgress; }
+
 private:
 
 	/** The core code of the fixup operation */
@@ -56,4 +59,8 @@ private:
 
 	/** Report any failures that may have happened during the rename */
 	void ReportFailures(const TArray<FRedirectorRefs>& RedirectorsToFix) const;
+
+private:
+
+	mutable bool bIsFixupReferencersInProgress = false;
 };
diff --git a/Engine/Source/Developer/AssetTools/Private/AssetTools.cpp b/Engine/Source/Developer/AssetTools/Private/AssetTools.cpp
--- a/Engine/Source/Developer/AssetTools/Private/AssetTools.cpp
+++ b/Engine/Source/Developer/AssetTools/Private/AssetTools.cpp
@@ -2924,6 +2924,11 @@ void UAssetToolsImpl::FixupReferencers(const TArray<UObjectRedirector*>& Objects
 	AssetFixUpRedirectors->FixupReferencers(Objects);
 }
 
+bool UAssetToolsImpl::IsFixupReferencersInProgress() const
+{
+	return AssetFixUpRedirectors->IsFixupReferencersInProgress();
+}
+
 void UAssetToolsImpl::OpenEditorForAssets(const TArray<UObject*>& Assets)
 {
 #if WITH_EDITOR
diff --git a/Engine/Source/Developer/AssetTools/Private/AssetTools.h b/Engine/Source/Developer/AssetTools/Private/AssetTools.h
--- a/Engine/Source/Developer/AssetTools/Private/AssetTools.h
+++ b/Engine/Source/Developer/AssetTools/Private/AssetTools.h
@@ -101,6 +101,7 @@ public:
 	virtual void MigratePackages(const TArray<FName>& PackageNamesToMigrate) const override;
 	virtual void BeginAdvancedCopyPackages(const TArray<FName>& InputNamesToCopy, const FString& TargetPath) const override;
 	virtual void FixupReferencers(const TArray<UObjectRedirector*>& Objects) const override;
+	virtual bool IsFixupReferencersInProgress() const override;
 	virtual FAssetPostRenameEvent& OnAssetPostRename() override { return AssetRenameManager->OnAssetPostRenameEvent(); }
 	virtual void ExpandDirectories(const TArray<FString>& Files, const FString& DestinationPath, TArray<TPair<FString, FString>>& FilesAndDestinations) const override;
 	virtual bool AdvancedCopyPackages(const FAdvancedCopyParams& CopyParams, const TArray<TMap<FString, FString>> PackagesAndDestinations) const override;
diff --git a/Engine/Source/Developer/AssetTools/Public/IAssetTools.h b/Engine/Source/Developer/AssetTools/Public/IAssetTools.h
--- a/Engine/Source/Developer/AssetTools/Public/IAssetTools.h
+++ b/Engine/Source/Developer/AssetTools/Public/IAssetTools.h
@@ -367,6 +367,9 @@ public:
 	/** Fix up references to the specified redirectors */
 	virtual void FixupReferencers(const TArray<UObjectRedirector*>& Objects) const = 0;
 
+	/** Returns whether redirectors are being fixed up. */
+	virtual bool IsFixupReferencersInProgress() const = 0;
+
 	/** Expands any folders found in the files list, and returns a flattened list of destination paths and files.  Mirrors directory structure. */
 	virtual void ExpandDirectories(const TArray<FString>& Files, const FString& DestinationPath, TArray<TPair<FString, FString>>& FilesAndDestinations) const = 0;
 
