Move the content browser skeletal mesh skinning icon from top right to bottom right corner.

diff --git a/Engine/Source/Developer/AssetTools/Private/AssetTypeActions/AssetTypeActions_SkeletalMesh.cpp b/Engine/Source/Developer/AssetTools/Private/AssetTypeActions/AssetTypeActions_SkeletalMesh.cpp
--- a/Engine/Source/Developer/AssetTools/Private/AssetTypeActions/AssetTypeActions_SkeletalMesh.cpp
+++ b/Engine/Source/Developer/AssetTools/Private/AssetTypeActions/AssetTypeActions_SkeletalMesh.cpp
@@ -620,9 +620,9 @@ TSharedPtr<SWidget> FAssetTypeActions_SkeletalMesh::GetThumbnailOverlay(const FA
 	return SNew(SBorder)
 		.BorderImage(FEditorStyle::GetNoBrush())
 		.Visibility(this, &FAssetTypeActions_SkeletalMesh::GetThumbnailSkinningOverlayVisibility, AssetData)
-		.Padding(FMargin(0.0f, 3.0f, 3.0f, 0.0f))
+		.Padding(FMargin(0.0f, 0.0f, 3.0f, 3.0f))
 		.HAlign(HAlign_Right)
-		.VAlign(VAlign_Top)
+		.VAlign(VAlign_Bottom)
 		[
 			SNew(SImage)
 			.ToolTipText(LOCTEXT("FAssetTypeActions_SkeletalMesh_NeedSkinning_ToolTip", "Asset geometry was imported, the skinning need to be validate"))
