Avoid multi-editing font face assets by default

This doesn't work with the font face editor, and can be done via the property matrix instead if needed

diff --git a/Engine/Source/Developer/AssetTools/Private/AssetTypeActions/AssetTypeActions_FontFace.cpp b/Engine/Source/Developer/AssetTools/Private/AssetTypeActions/AssetTypeActions_FontFace.cpp
--- a/Engine/Source/Developer/AssetTools/Private/AssetTypeActions/AssetTypeActions_FontFace.cpp
+++ b/Engine/Source/Developer/AssetTools/Private/AssetTypeActions/AssetTypeActions_FontFace.cpp
@@ -28,7 +28,13 @@ void FAssetTypeActions_FontFace::OpenAssetEditor( const TArray<UObject*>& InObje
 {
 	// Load the FontEditor module to ensure that FFontFaceDetailsCustomization is registered
 	IFontEditorModule* FontEditorModule = &FModuleManager::LoadModuleChecked<IFontEditorModule>("FontEditor");
-	FSimpleAssetEditor::CreateEditor(EToolkitMode::Standalone, EditWithinLevelEditor, InObjects);
+	
+	// Open each object in turn, as the default editor would try and collapse a multi-edit together 
+	// into a single editor instance which doesn't really work for font face assets
+	for (UObject* Object : InObjects)
+	{
+		FSimpleAssetEditor::CreateEditor(EToolkitMode::Standalone, EditWithinLevelEditor, Object);
+	}
 }
 
 bool FAssetTypeActions_FontFace::CanExecuteReimport(const TArray<TWeakObjectPtr<UFontFace>> Objects) const
