Added DuplicateAssetWithDialogAndTitle() to IAssetTools, which allows for the user to specify a title for the asset saving dialog that is displayed.

diff --git a/Engine/Source/Developer/AssetTools/Private/AssetTools.cpp b/Engine/Source/Developer/AssetTools/Private/AssetTools.cpp
--- a/Engine/Source/Developer/AssetTools/Private/AssetTools.cpp
+++ b/Engine/Source/Developer/AssetTools/Private/AssetTools.cpp
@@ -548,9 +548,14 @@ UObject* UAssetToolsImpl::CreateAssetWithDialog(const FString& AssetName, const
 }
 
 UObject* UAssetToolsImpl::DuplicateAssetWithDialog(const FString& AssetName, const FString& PackagePath, UObject* OriginalObject)
+{
+	return DuplicateAssetWithDialogAndTitle(AssetName, PackagePath, OriginalObject, LOCTEXT("DuplicateAssetDialogTitle", "Duplicate Asset As"));
+}
+
+UObject* UAssetToolsImpl::DuplicateAssetWithDialogAndTitle(const FString& AssetName, const FString& PackagePath, UObject* OriginalObject, FText DialogTitle)
 {
 	FSaveAssetDialogConfig SaveAssetDialogConfig;
-	SaveAssetDialogConfig.DialogTitleOverride = LOCTEXT("DuplicateAssetDialogTitle", "Duplicate Asset As");
+	SaveAssetDialogConfig.DialogTitleOverride = DialogTitle;
 	SaveAssetDialogConfig.DefaultPath = PackagePath;
 	SaveAssetDialogConfig.DefaultAssetName = AssetName;
 	SaveAssetDialogConfig.ExistingAssetPolicy = ESaveAssetDialogExistingAssetPolicy::AllowButWarn;
diff --git a/Engine/Source/Developer/AssetTools/Private/AssetTools.h b/Engine/Source/Developer/AssetTools/Private/AssetTools.h
--- a/Engine/Source/Developer/AssetTools/Private/AssetTools.h
+++ b/Engine/Source/Developer/AssetTools/Private/AssetTools.h
@@ -79,6 +79,7 @@ public:
 	virtual UObject* CreateAssetWithDialog(const FString& AssetName, const FString& PackagePath, UClass* AssetClass, UFactory* Factory, FName CallingContext = NAME_None) override;
 	virtual UObject* DuplicateAsset(const FString& AssetName, const FString& PackagePath, UObject* OriginalObject) override;
 	virtual UObject* DuplicateAssetWithDialog(const FString& AssetName, const FString& PackagePath, UObject* OriginalObject) override;
+	virtual UObject* DuplicateAssetWithDialogAndTitle(const FString& AssetName, const FString& PackagePath, UObject* OriginalObject, FText DialogTitle) override;
 	virtual bool RenameAssets(const TArray<FAssetRenameData>& AssetsAndNames) override;
 	virtual void RenameAssetsWithDialog(const TArray<FAssetRenameData>& AssetsAndNames, bool bAutoCheckout = false) override;
 	virtual void FindSoftReferencesToObject(FSoftObjectPath TargetObject, TArray<UObject*>& ReferencingObjects) override;
diff --git a/Engine/Source/Developer/AssetTools/Public/IAssetTools.h b/Engine/Source/Developer/AssetTools/Public/IAssetTools.h
--- a/Engine/Source/Developer/AssetTools/Public/IAssetTools.h
+++ b/Engine/Source/Developer/AssetTools/Public/IAssetTools.h
@@ -222,6 +222,13 @@ public:
 	UFUNCTION(BlueprintCallable, Category = "Editor Scripting | Asset Tools")
 	virtual UObject* DuplicateAssetWithDialog(const FString& AssetName, const FString& PackagePath, UObject* OriginalObject) = 0;
 
+	/** Opens an asset picker dialog and creates an asset with the specified name and path. 
+	 * Uses OriginalObject as the duplication source.
+	 * Uses DialogTitle as the dialog's title.
+	 */
+	UFUNCTION(BlueprintCallable, Category = "Editor Scripting | Asset Tools")
+	virtual UObject* DuplicateAssetWithDialogAndTitle(const FString& AssetName, const FString& PackagePath, UObject* OriginalObject, FText DialogTitle) = 0;
+
 	/** Creates an asset with the specified name and path. Uses OriginalObject as the duplication source. */
 	UFUNCTION(BlueprintCallable, Category="Editor Scripting | Asset Tools")
 	virtual UObject* DuplicateAsset(const FString& AssetName, const FString& PackagePath, UObject* OriginalObject) = 0;
