Migrate now correctly checks for packages outside the source packages' root directory, and will not add those to the dependency list.
#rnx

diff --git a/Engine/Source/Developer/AssetTools/Private/AssetTools.cpp b/Engine/Source/Developer/AssetTools/Private/AssetTools.cpp
--- a/Engine/Source/Developer/AssetTools/Private/AssetTools.cpp
+++ b/Engine/Source/Developer/AssetTools/Private/AssetTools.cpp
@@ -2580,7 +2580,12 @@ void UAssetToolsImpl::PerformMigratePackages(TArray<FName> PackageNamesToMigrate
 			if ( !AllPackageNamesToMove.Contains(*PackageIt) )
 			{
 				AllPackageNamesToMove.Add(*PackageIt);
-				RecursiveGetDependencies(*PackageIt, AllPackageNamesToMove);
+				FString Path = (*PackageIt).ToString();
+				FString OriginalRootString;
+				Path.RemoveFromStart(TEXT("/"));
+				Path.Split("/", &OriginalRootString, &Path, ESearchCase::IgnoreCase, ESearchDir::FromStart);
+				OriginalRootString = TEXT("/") + OriginalRootString;
+				RecursiveGetDependencies(*PackageIt, AllPackageNamesToMove, OriginalRootString);
 			}
 		}
 	}
@@ -2837,7 +2842,7 @@ void UAssetToolsImpl::MigratePackages_ReportConfirmed(TSharedPtr<TArray<ReportPa
 	MigrateLog.Notify(LogMessage, Severity, true);
 }
 
-void UAssetToolsImpl::RecursiveGetDependencies(const FName& PackageName, TSet<FName>& AllDependencies) const
+void UAssetToolsImpl::RecursiveGetDependencies(const FName& PackageName, TSet<FName>& AllDependencies, const FString& OriginalRoot) const
 {
 	FAssetRegistryModule& AssetRegistryModule = FModuleManager::Get().LoadModuleChecked<FAssetRegistryModule>(TEXT("AssetRegistry"));
 	TArray<FName> Dependencies;
@@ -2847,13 +2852,14 @@ void UAssetToolsImpl::RecursiveGetDependencies(const FName& PackageName, TSet<FN
 	{
 		if ( !AllDependencies.Contains(*DependsIt) )
 		{
-			// @todo Make this skip all packages whose root is different than the source package list root. For now we just skip engine content.
 			const bool bIsEnginePackage = (*DependsIt).ToString().StartsWith(TEXT("/Engine"));
 			const bool bIsScriptPackage = (*DependsIt).ToString().StartsWith(TEXT("/Script"));
-			if ( !bIsEnginePackage && !bIsScriptPackage )
+			// Skip all packages whose root is different than the source package list root
+			const bool bIsInSamePackage = (*DependsIt).ToString().StartsWith(OriginalRoot);
+			if ( !bIsEnginePackage && !bIsScriptPackage && bIsInSamePackage )
 			{
 				AllDependencies.Add(*DependsIt);
-				RecursiveGetDependencies(*DependsIt, AllDependencies);
+				RecursiveGetDependencies(*DependsIt, AllDependencies, OriginalRoot);
 			}
 		}
 	}
diff --git a/Engine/Source/Developer/AssetTools/Private/AssetTools.h b/Engine/Source/Developer/AssetTools/Private/AssetTools.h
--- a/Engine/Source/Developer/AssetTools/Private/AssetTools.h
+++ b/Engine/Source/Developer/AssetTools/Private/AssetTools.h
@@ -157,7 +157,7 @@ private:
 	void AdvancedCopyPackages_ReportConfirmed(FAdvancedCopyParams CopyParam, TArray<TMap<FString, FString>> DestinationMap) const;
 
 	/** Gets the dependencies of the specified package recursively */
-	void RecursiveGetDependencies(const FName& PackageName, TSet<FName>& AllDependencies) const;
+	void RecursiveGetDependencies(const FName& PackageName, TSet<FName>& AllDependencies, const FString& OriginalRoot) const;
 
 	/** Gets the dependencies of the specified package recursively while omitting things that don't pass the FARFilter passed in from FAdvancedCopyParams */
 	void RecursiveGetDependenciesAdvanced(const FName& PackageName, FAdvancedCopyParams& CopyParams, TArray<FName>& AllDependencies, TMap<FName, FName>& DependencyMap, const class UAdvancedCopyCustomization* CopyCustomization, TArray<FAssetData>& OptionalAssetData) const;
