Fast and lexical sort functors for FSoftObjectPath, FSoftObjectPtr, TSoftObjectPtr and TSoftClassPtr.

diff --git a/Engine/Source/Developer/AssetTools/Private/AssetRenameManager.cpp b/Engine/Source/Developer/AssetTools/Private/AssetRenameManager.cpp
--- a/Engine/Source/Developer/AssetTools/Private/AssetRenameManager.cpp
+++ b/Engine/Source/Developer/AssetTools/Private/AssetRenameManager.cpp
@@ -1170,15 +1170,6 @@ bool FAssetRenameManager::CheckPackageForSoftObjectReferences(UPackage* Package,
 {
 	using namespace AssetRenameManagerImpl;
 
-	struct FSoftObjectPathLess
-	{
-		FORCEINLINE bool operator()(const FSoftObjectPath& A, const FSoftObjectPath& B) const
-		{
-			int32 Result = A.GetAssetPathName().CompareIndexes(B.GetAssetPathName());
-			return Result ? Result < 0 : A.GetSubPathString() < B.GetSubPathString();
-		}
-	};
-
 	bool bFoundReference = false;
 
 	// First check cache
@@ -1217,7 +1208,7 @@ bool FAssetRenameManager::CheckPackageForSoftObjectReferences(UPackage* Package,
 		CachedReferences->Map.GenerateKeyArray(CachedReferences->Keys);
 		
 		// Keys need to be sorted for binary search
-		CachedReferences->Keys.Sort(FSoftObjectPathLess());
+		CachedReferences->Keys.Sort(FSoftObjectPathFastLess());
 	}
 
 	for (const TPair<FSoftObjectPath, FSoftObjectPath>& Pair : AssetRedirectorMap)
@@ -1225,7 +1216,7 @@ bool FAssetRenameManager::CheckPackageForSoftObjectReferences(UPackage* Package,
 		const FString& CheckSubPath = Pair.Key.GetSubPathString();
 
 		// Find where we're going to start iterating
-		int32 Index = Algo::LowerBound(CachedReferences->Keys, Pair.Key, FSoftObjectPathLess());
+		int32 Index = Algo::LowerBound(CachedReferences->Keys, Pair.Key, FSoftObjectPathFastLess());
 		for (int32 Num = CachedReferences->Keys.Num(); Index < Num; ++Index)
 		{
 			const FSoftObjectPath& CachedKey = CachedReferences->Keys[Index];
diff --git a/Engine/Source/Runtime/CoreUObject/Public/UObject/SoftObjectPath.h b/Engine/Source/Runtime/CoreUObject/Public/UObject/SoftObjectPath.h
--- a/Engine/Source/Runtime/CoreUObject/Public/UObject/SoftObjectPath.h
+++ b/Engine/Source/Runtime/CoreUObject/Public/UObject/SoftObjectPath.h
@@ -217,6 +217,34 @@ private:
 	friend struct Z_Construct_UScriptStruct_FSoftObjectPath_Statics;
 };
 
+/** Fast non-alphabetical order that is only stable during this process' lifetime. */
+struct FSoftObjectPathFastLess
+{
+	bool operator()(const FSoftObjectPath& Lhs, const FSoftObjectPath& Rhs) const
+	{
+		int32 Comp = Lhs.GetAssetPathName().CompareIndexes(Rhs.GetAssetPathName());
+		if (Comp < 0)
+		{
+			return true;
+		}
+		return Comp == 0 && Lhs.GetSubPathString() < Rhs.GetSubPathString();
+	}
+};
+
+/** Slow alphabetical order that is stable / deterministic over process runs. */
+struct FSoftObjectPathLexicalLess
+{
+	bool operator()(const FSoftObjectPath& Lhs, const FSoftObjectPath& Rhs) const
+	{
+		int32 Comp = Lhs.GetAssetPathName().Compare(Rhs.GetAssetPathName());
+		if (Comp < 0)
+		{
+			return true;
+		}
+		return Comp == 0 && Lhs.GetSubPathString() < Rhs.GetSubPathString();
+	}
+};
+
 /**
  * A struct that contains a string reference to a class, can be used to make soft references to classes
  */
diff --git a/Engine/Source/Runtime/CoreUObject/Public/UObject/SoftObjectPtr.h b/Engine/Source/Runtime/CoreUObject/Public/UObject/SoftObjectPtr.h
--- a/Engine/Source/Runtime/CoreUObject/Public/UObject/SoftObjectPtr.h
+++ b/Engine/Source/Runtime/CoreUObject/Public/UObject/SoftObjectPtr.h
@@ -620,6 +620,26 @@ private:
 template <class T> struct TIsPODType<TSoftClassPtr<T> > { enum { Value = TIsPODType<FSoftObjectPtr>::Value }; };
 template <class T> struct TIsWeakPointerType<TSoftClassPtr<T> > { enum { Value = TIsWeakPointerType<FSoftObjectPtr>::Value }; };
 
+/** Fast non-alphabetical order that is only stable during this process' lifetime. */
+struct FSoftObjectPtrFastLess : private FSoftObjectPathFastLess
+{
+	template <typename SoftObjectPtrType>
+	bool operator()(const SoftObjectPtrType& Lhs, const SoftObjectPtrType& Rhs) const
+	{
+		return FSoftObjectPathFastLess::operator()(Lhs.ToSoftObjectPath(), Rhs.ToSoftObjectPath());
+	}
+};
+
+/** Slow alphabetical order that is stable / deterministic over process runs. */
+struct FSoftObjectPtrLexicalLess : private FSoftObjectPathLexicalLess
+{
+	template <typename SoftObjectPtrType>
+	bool operator()(const SoftObjectPtrType& Lhs, const SoftObjectPtrType& Rhs) const
+	{
+		return FSoftObjectPathLexicalLess::operator()(Lhs.ToSoftObjectPath(), Rhs.ToSoftObjectPath());
+	}
+};
+
 UE_DEPRECATED(4.18, "FAssetPtr was renamed to FSoftObjectPtr as it is not necessarily an asset")
 typedef FSoftObjectPtr FAssetPtr;
 
