Fix GetTypedObjectPtrs templetization to something that can actually compile.

diff --git a/Engine/Source/Developer/AssetTools/Public/AssetTypeActions_Base.h b/Engine/Source/Developer/AssetTools/Public/AssetTypeActions_Base.h
--- a/Engine/Source/Developer/AssetTools/Public/AssetTypeActions_Base.h
+++ b/Engine/Source/Developer/AssetTools/Public/AssetTypeActions_Base.h
@@ -207,11 +207,11 @@ protected:
 	}
 
 	template <typename T>
-	static TArray<T> GetTypedObjectPtrs(const TArray<UObject*>& InObjects)
+	static TArray<T*> GetTypedObjectPtrs(const TArray<UObject*>& InObjects)
 	{
 		check(InObjects.Num() > 0);
 
-		TArray<T> TypedObjects;
+		TArray<T*> TypedObjects;
 		for (UObject* ObjIt : InObjects)
 		{
 			TypedObjects.Add(CastChecked<T>(ObjIt));
