Fix asset bulk export not having filename
#rnx

diff --git a/Engine/Source/Developer/AssetTools/Private/AssetTools.cpp b/Engine/Source/Developer/AssetTools/Private/AssetTools.cpp
--- a/Engine/Source/Developer/AssetTools/Private/AssetTools.cpp
+++ b/Engine/Source/Developer/AssetTools/Private/AssetTools.cpp
@@ -2294,7 +2294,7 @@ void UAssetToolsImpl::ExportAssetsInternal(const TArray<UObject*>& ObjectsToExpo
 				if (PackageName.Left(1) == TEXT("/"))
 				{
 					// Trim the leading slash so the file manager doesn't get confused
-					PackageName.MidInline(1, false);
+					PackageName.MidInline(1, MAX_int32, false);
 				}
 
 				FPaths::NormalizeFilename(PackageName);
