Add option on whether to prompt user with checkout dialog in AssetFixUpRedirectors and use it to avoid pop-up in PPID generation.

diff --git a/Engine/Source/Developer/AssetTools/Private/AssetFixUpRedirectors.cpp b/Engine/Source/Developer/AssetTools/Private/AssetFixUpRedirectors.cpp
--- a/Engine/Source/Developer/AssetTools/Private/AssetFixUpRedirectors.cpp
+++ b/Engine/Source/Developer/AssetTools/Private/AssetFixUpRedirectors.cpp
@@ -43,7 +43,7 @@ struct FRedirectorRefs
 };
 
 
-void FAssetFixUpRedirectors::FixupReferencers(const TArray<UObjectRedirector*>& Objects) const
+void FAssetFixUpRedirectors::FixupReferencers(const TArray<UObjectRedirector*>& Objects, const bool bCheckoutDialogPrompt) const
 {
 	// Transform array into TWeakObjectPtr array
 	TArray<TWeakObjectPtr<UObjectRedirector>> ObjectWeakPtrs;
@@ -60,18 +60,18 @@ void FAssetFixUpRedirectors::FixupReferencers(const TArray<UObjectRedirector*>&
 		{
 			// Open a dialog asking the user to wait while assets are being discovered
 			SDiscoveringAssetsDialog::OpenDiscoveringAssetsDialog(
-				SDiscoveringAssetsDialog::FOnAssetsDiscovered::CreateSP(this, &FAssetFixUpRedirectors::ExecuteFixUp, ObjectWeakPtrs)
+				SDiscoveringAssetsDialog::FOnAssetsDiscovered::CreateSP(this, &FAssetFixUpRedirectors::ExecuteFixUp, ObjectWeakPtrs, bCheckoutDialogPrompt)
 				);
 		}
 		else
 		{
 			// No need to wait, attempt to fix references now.
-			ExecuteFixUp(ObjectWeakPtrs);
+			ExecuteFixUp(ObjectWeakPtrs, bCheckoutDialogPrompt);
 		}
 	}
 }
 
-void FAssetFixUpRedirectors::ExecuteFixUp(TArray<TWeakObjectPtr<UObjectRedirector>> Objects) const
+void FAssetFixUpRedirectors::ExecuteFixUp(TArray<TWeakObjectPtr<UObjectRedirector>> Objects, const bool bCheckoutDialogPrompt) const
 {
 	TGuardValue<bool> Guard(bIsFixupReferencersInProgress, true);
 
@@ -98,8 +98,8 @@ void FAssetFixUpRedirectors::ExecuteFixUp(TArray<TWeakObjectPtr<UObjectRedirecto
 			TArray<UPackage*> ReferencingPackagesToSave;
 			LoadReferencingPackages(RedirectorRefsList, ReferencingPackagesToSave);
 
-			// Prompt to check out all referencing packages, leave redirectors for assets referenced by packages that are not checked out and remove those packages from the save list.
-			const bool bUserAcceptedCheckout = CheckOutReferencingPackages(RedirectorRefsList, ReferencingPackagesToSave);
+			// Check out all referencing packages, leave redirectors for assets referenced by packages that are not checked out and remove those packages from the save list.
+			const bool bUserAcceptedCheckout = CheckOutReferencingPackages(RedirectorRefsList, ReferencingPackagesToSave, bCheckoutDialogPrompt);
 			if ( bUserAcceptedCheckout )
 			{
 				// If any referencing packages are left read-only, the checkout failed or SCC was not enabled. Trim them from the save list and leave redirectors.
@@ -221,7 +221,7 @@ void FAssetFixUpRedirectors::LoadReferencingPackages(TArray<FRedirectorRefs>& Re
 	}
 }
 
-bool FAssetFixUpRedirectors::CheckOutReferencingPackages(TArray<FRedirectorRefs>& RedirectorsToFix, TArray<UPackage*>& InOutReferencingPackagesToSave) const
+bool FAssetFixUpRedirectors::CheckOutReferencingPackages(TArray<FRedirectorRefs>& RedirectorsToFix, TArray<UPackage*>& InOutReferencingPackagesToSave, const bool bCheckoutDialogPrompt) const
 {
 	// Prompt to check out all successfully loaded packages
 	bool bUserAcceptedCheckout = true;
@@ -232,7 +232,17 @@ bool FAssetFixUpRedirectors::CheckOutReferencingPackages(TArray<FRedirectorRefs>
 		{
 			TArray<UPackage*> PackagesCheckedOutOrMadeWritable;
 			TArray<UPackage*> PackagesNotNeedingCheckout;
-			bUserAcceptedCheckout = FEditorFileUtils::PromptToCheckoutPackages( false, InOutReferencingPackagesToSave, &PackagesCheckedOutOrMadeWritable, &PackagesNotNeedingCheckout );
+			if (bCheckoutDialogPrompt)
+			{
+				bUserAcceptedCheckout = FEditorFileUtils::PromptToCheckoutPackages( false, InOutReferencingPackagesToSave, &PackagesCheckedOutOrMadeWritable, &PackagesNotNeedingCheckout );
+			}
+			else
+			{
+				const bool bErrorIfAlreadyCheckedOut = false;
+				const bool bConfirmPackageBranchCheckOutStatus = false;
+				FEditorFileUtils::CheckoutPackages(InOutReferencingPackagesToSave, &PackagesCheckedOutOrMadeWritable, bErrorIfAlreadyCheckedOut, bConfirmPackageBranchCheckOutStatus);
+			}
+
 			if ( bUserAcceptedCheckout )
 			{
 				TArray<UPackage*> PackagesThatCouldNotBeCheckedOut = InOutReferencingPackagesToSave;
diff --git a/Engine/Source/Developer/AssetTools/Private/AssetFixUpRedirectors.h b/Engine/Source/Developer/AssetTools/Private/AssetFixUpRedirectors.h
--- a/Engine/Source/Developer/AssetTools/Private/AssetFixUpRedirectors.h
+++ b/Engine/Source/Developer/AssetTools/Private/AssetFixUpRedirectors.h
@@ -9,8 +9,11 @@ struct FRedirectorRefs;
 class FAssetFixUpRedirectors : public TSharedFromThis<FAssetFixUpRedirectors>
 {
 public:
-	/** Fix up references to the specified redirectors */
-	void FixupReferencers(const TArray<UObjectRedirector*>& Objects) const;
+	/**
+	 * Fix up references to the specified redirectors.
+	 * @param bCheckoutDialogPrompt indicates whether to prompt the user with files checkout dialog or silently attempt to checkout all necessary files.
+	 */
+	void FixupReferencers(const TArray<UObjectRedirector*>& Objects, bool bCheckoutDialogPrompt = true) const;
 
 	/** Returns whether redirectors are being fixed up. */
 	bool IsFixupReferencersInProgress() const { return bIsFixupReferencersInProgress; }
@@ -18,7 +21,7 @@ public:
 private:
 
 	/** The core code of the fixup operation */
-	void ExecuteFixUp(TArray<TWeakObjectPtr<UObjectRedirector>> Objects) const;
+	void ExecuteFixUp(TArray<TWeakObjectPtr<UObjectRedirector>> Objects, bool bCheckoutDialogPrompt) const;
 
 	/** Fills out the Referencing packages for all the redirectors described in AssetsToPopulate */
 	void PopulateRedirectorReferencers(TArray<FRedirectorRefs>& RedirectorsToPopulate) const;
@@ -33,11 +36,11 @@ private:
 	void LoadReferencingPackages(TArray<FRedirectorRefs>& RedirectorsToFix, TArray<UPackage*>& OutReferencingPackagesToSave) const;
 
 	/** 
-	  * Prompts to check out referencing packages and marks assets whose referencing packages were not checked out to not fix the redirector.
+	  * Check out referencing packages and marks assets whose referencing packages were not checked out to not fix the redirector.
 	  * Trims PackagesToSave when necessary.
 	  * Returns true if the user opted to continue the operation or no dialog was required.
 	  */
-	bool CheckOutReferencingPackages(TArray<FRedirectorRefs>& RedirectorsToFix, TArray<UPackage*>& InOutReferencingPackagesToSave) const;
+	bool CheckOutReferencingPackages(TArray<FRedirectorRefs>& RedirectorsToFix, TArray<UPackage*>& InOutReferencingPackagesToSave, bool bCheckoutDialogPrompt) const;
 
 	/** Finds any read only packages and removes them from the save list. Redirectors referenced by these packages will not be fixed up. */ 
 	void DetectReadOnlyPackages(TArray<FRedirectorRefs>& RedirectorsToFix, TArray<UPackage*>& InOutReferencingPackagesToSave) const;
diff --git a/Engine/Source/Developer/AssetTools/Private/AssetTools.cpp b/Engine/Source/Developer/AssetTools/Private/AssetTools.cpp
--- a/Engine/Source/Developer/AssetTools/Private/AssetTools.cpp
+++ b/Engine/Source/Developer/AssetTools/Private/AssetTools.cpp
@@ -2937,9 +2937,9 @@ void UAssetToolsImpl::RecursiveGetDependenciesAdvanced(const FName& PackageName,
 	}
 }
 
-void UAssetToolsImpl::FixupReferencers(const TArray<UObjectRedirector*>& Objects) const
+void UAssetToolsImpl::FixupReferencers(const TArray<UObjectRedirector*>& Objects, bool bCheckoutDialogPrompt) const
 {
-	AssetFixUpRedirectors->FixupReferencers(Objects);
+	AssetFixUpRedirectors->FixupReferencers(Objects, bCheckoutDialogPrompt);
 }
 
 bool UAssetToolsImpl::IsFixupReferencersInProgress() const
diff --git a/Engine/Source/Developer/AssetTools/Private/AssetTools.h b/Engine/Source/Developer/AssetTools/Private/AssetTools.h
--- a/Engine/Source/Developer/AssetTools/Private/AssetTools.h
+++ b/Engine/Source/Developer/AssetTools/Private/AssetTools.h
@@ -102,7 +102,7 @@ public:
 	virtual bool CreateDiffProcess(const FString& DiffCommand, const FString& OldTextFilename, const FString& NewTextFilename, const FString& DiffArgs = FString("")) const override;
 	virtual void MigratePackages(const TArray<FName>& PackageNamesToMigrate) const override;
 	virtual void BeginAdvancedCopyPackages(const TArray<FName>& InputNamesToCopy, const FString& TargetPath) const override;
-	virtual void FixupReferencers(const TArray<UObjectRedirector*>& Objects) const override;
+	virtual void FixupReferencers(const TArray<UObjectRedirector*>& Objects, bool bCheckoutDialogPrompt = true) const override;
 	virtual bool IsFixupReferencersInProgress() const override;
 	virtual FAssetPostRenameEvent& OnAssetPostRename() override { return AssetRenameManager->OnAssetPostRenameEvent(); }
 	virtual void ExpandDirectories(const TArray<FString>& Files, const FString& DestinationPath, TArray<TPair<FString, FString>>& FilesAndDestinations) const override;
diff --git a/Engine/Source/Developer/AssetTools/Public/IAssetTools.h b/Engine/Source/Developer/AssetTools/Public/IAssetTools.h
--- a/Engine/Source/Developer/AssetTools/Public/IAssetTools.h
+++ b/Engine/Source/Developer/AssetTools/Public/IAssetTools.h
@@ -373,8 +373,11 @@ public:
 	/* Copy packages and dependencies to another folder */
 	virtual void BeginAdvancedCopyPackages(const TArray<FName>& InputNamesToCopy, const FString& TargetPath) const = 0;
 
-	/** Fix up references to the specified redirectors */
-	virtual void FixupReferencers(const TArray<UObjectRedirector*>& Objects) const = 0;
+	/**
+	 * Fix up references to the specified redirectors.
+	 * @param bCheckoutDialogPrompt indicates whether to prompt the user with files checkout dialog or silently attempt to checkout all necessary files.
+	 */
+	virtual void FixupReferencers(const TArray<UObjectRedirector*>& Objects, bool bCheckoutDialogPrompt = true) const = 0;
 
 	/** Returns whether redirectors are being fixed up. */
 	virtual bool IsFixupReferencersInProgress() const = 0;
