Integrating CL#11136074 from Dev-Enterprise
Fixing Python asset import task returning wrong path when not combining mesh.

diff --git a/Engine/Source/Developer/AssetTools/Private/AssetTools.cpp b/Engine/Source/Developer/AssetTools/Private/AssetTools.cpp
--- a/Engine/Source/Developer/AssetTools/Private/AssetTools.cpp
+++ b/Engine/Source/Developer/AssetTools/Private/AssetTools.cpp
@@ -2050,7 +2050,7 @@ TArray<UObject*> UAssetToolsImpl::ImportAssetsInternal(const TArray<FString>& Fi
 
 					for (UObject* AdditionalResult : Factory->GetAdditionalImportedObjects())
 					{
-						ReturnObjects.Add(Result);
+						ReturnObjects.Add(AdditionalResult);
 					}
 
 					bImportSucceeded = true;
