PR #6761: Fix for unsafe cast of IAssetTypeActions* to FAssetTypeActions_Base* (Contributed by [REDACTED_PERSON])

diff --git a/Engine/Source/Developer/AssetTools/Public/AssetTypeActions_Base.h b/Engine/Source/Developer/AssetTools/Public/AssetTypeActions_Base.h
--- a/Engine/Source/Developer/AssetTools/Public/AssetTypeActions_Base.h
+++ b/Engine/Source/Developer/AssetTools/Public/AssetTypeActions_Base.h
@@ -162,7 +162,7 @@ PRAGMA_ENABLE_DEPRECATION_WARNINGS
 		InFilter.bRecursiveClasses = true;
 	}
 	
-	virtual FText GetDisplayNameFromAssetData(const FAssetData& AssetData) const
+	virtual FText GetDisplayNameFromAssetData(const FAssetData& AssetData) const override
 	{
 		return FText::GetEmpty();
 	}
diff --git a/Engine/Source/Developer/AssetTools/Public/IAssetTypeActions.h b/Engine/Source/Developer/AssetTools/Public/IAssetTypeActions.h
--- a/Engine/Source/Developer/AssetTools/Public/IAssetTypeActions.h
+++ b/Engine/Source/Developer/AssetTools/Public/IAssetTypeActions.h
@@ -120,6 +120,9 @@ public:
 	/** Builds the filter for this class*/
 	virtual void BuildBackendFilter(struct FARFilter& InFilter) = 0;
 
+	/** Optionally gets a class display name for this asset (otherwise, returns empty text (e.g. `FText::GetEmpty()`) */
+	virtual FText GetDisplayNameFromAssetData(const FAssetData& AssetData) const = 0;
+
 	/** Sets whether or not this asset type is a supported type for this editor session. */
 	virtual void SetSupported(bool bInSupported) = 0;
 
diff --git a/Engine/Source/Editor/UnrealEd/Private/AssetThumbnail.cpp b/Engine/Source/Editor/UnrealEd/Private/AssetThumbnail.cpp
--- a/Engine/Source/Editor/UnrealEd/Private/AssetThumbnail.cpp
+++ b/Engine/Source/Editor/UnrealEd/Private/AssetThumbnail.cpp
@@ -585,11 +585,7 @@ private:
 			{
 				if (AssetData != nullptr)
 				{
-					FAssetTypeActions_Base* BaseAssetTypeAction = static_cast<FAssetTypeActions_Base*>(AssetTypeActions.Pin().Get());
-					if (BaseAssetTypeAction != nullptr)
-					{
-						ClassDisplayName = BaseAssetTypeAction->GetDisplayNameFromAssetData(*AssetData);
-					}
+					ClassDisplayName = AssetTypeActions.Pin()->GetDisplayNameFromAssetData(*AssetData);
 				}
 
 				if (ClassDisplayName.IsEmpty())
