Moved CollectionAssetManagement to AssetTools so it can be shared

#rnx

diff --git a/Engine/Source/Developer/AssetTools/Private/CollectionAssetManagement.cpp b/Engine/Source/Developer/AssetTools/Private/CollectionAssetManagement.cpp
new file mode 100644
--- /dev/null
+++ b/Engine/Source/Developer/AssetTools/Private/CollectionAssetManagement.cpp
@@ -0,0 +1,261 @@
+// Copyright Epic Games, Inc. All Rights Reserved.
+
+#include "CollectionAssetManagement.h"
+#include "ISourceControlProvider.h"
+#include "ISourceControlModule.h"
+#include "ICollectionManager.h"
+#include "CollectionManagerModule.h"
+#include "Framework/Notifications/NotificationManager.h"
+#include "Widgets/Notifications/SNotificationList.h"
+
+#define LOCTEXT_NAMESPACE "ContentBrowser"
+
+FCollectionAssetManagement::FCollectionAssetManagement()
+{
+	FCollectionManagerModule& CollectionManagerModule = FCollectionManagerModule::GetModule();
+
+	// Register the notifications we need in order to keep things up-to-date
+	OnCollectionRenamedHandle = CollectionManagerModule.Get().OnCollectionRenamed().AddRaw(this, &FCollectionAssetManagement::HandleCollectionRenamed);
+	OnCollectionDestroyedHandle = CollectionManagerModule.Get().OnCollectionDestroyed().AddRaw(this, &FCollectionAssetManagement::HandleCollectionDestroyed);
+	OnCollectionUpdatedHandle = CollectionManagerModule.Get().OnCollectionUpdated().AddRaw(this, &FCollectionAssetManagement::HandleCollectionUpdated);
+	OnAssetsAddedHandle = CollectionManagerModule.Get().OnAssetsAdded().AddRaw(this, &FCollectionAssetManagement::HandleAssetsAddedToCollection);
+	OnAssetsRemovedHandle = CollectionManagerModule.Get().OnAssetsRemoved().AddRaw(this, &FCollectionAssetManagement::HandleAssetsRemovedFromCollection);
+}
+
+FCollectionAssetManagement::~FCollectionAssetManagement()
+{
+	// Check IsModuleAvailable as we might be in the process of shutting down...
+	if (FCollectionManagerModule::IsModuleAvailable())
+	{
+		FCollectionManagerModule& CollectionManagerModule = FCollectionManagerModule::GetModule();
+
+		CollectionManagerModule.Get().OnCollectionRenamed().Remove(OnCollectionRenamedHandle);
+		CollectionManagerModule.Get().OnCollectionDestroyed().Remove(OnCollectionDestroyedHandle);
+		CollectionManagerModule.Get().OnCollectionUpdated().Remove(OnCollectionUpdatedHandle);
+		CollectionManagerModule.Get().OnAssetsAdded().Remove(OnAssetsAddedHandle);
+		CollectionManagerModule.Get().OnAssetsRemoved().Remove(OnAssetsRemovedHandle);
+	}
+}
+
+void FCollectionAssetManagement::SetCurrentAssets(const TArray<FAssetData>& CurrentAssets)
+{
+	CurrentAssetPaths.Empty();
+	for (const FAssetData& AssetData : CurrentAssets)
+	{
+		CurrentAssetPaths.Add(AssetData.ObjectPath);
+	}
+
+	UpdateAssetManagementState();
+}
+
+void FCollectionAssetManagement::AddCurrentAssetsToCollection(FCollectionNameType InCollectionKey)
+{
+	FCollectionManagerModule& CollectionManagerModule = FCollectionManagerModule::GetModule();
+
+	const TArray<FName> ObjectPaths = CurrentAssetPaths.Array();
+
+	FText ResultText;
+	bool bSuccess = false;
+	{
+		int32 NumAdded = 0;
+		if (CollectionManagerModule.Get().AddToCollection(InCollectionKey.Name, InCollectionKey.Type, ObjectPaths, &NumAdded))
+		{
+			bSuccess = true;
+
+			FFormatNamedArguments Args;
+			Args.Add(TEXT("Number"), NumAdded);
+			Args.Add(TEXT("CollectionName"), FText::FromName(InCollectionKey.Name));
+			ResultText = FText::Format(LOCTEXT("CollectionAssetsAdded", "Added {Number} asset(s) to {CollectionName}"), Args);
+		}
+		else
+		{
+			ResultText = CollectionManagerModule.Get().GetLastError();
+		}
+	}
+
+	if (!ResultText.IsEmpty())
+	{
+		FNotificationInfo Info(ResultText);
+		Info.bFireAndForget = true;
+		Info.bUseLargeFont = false;
+
+		TSharedPtr<SNotificationItem> Item = FSlateNotificationManager::Get().AddNotification(Info);
+		if (Item.IsValid())
+		{
+			Item->SetCompletionState((bSuccess) ? SNotificationItem::CS_Success : SNotificationItem::CS_Fail);
+		}
+	}
+}
+
+void FCollectionAssetManagement::RemoveCurrentAssetsFromCollection(FCollectionNameType InCollectionKey)
+{
+	FCollectionManagerModule& CollectionManagerModule = FCollectionManagerModule::GetModule();
+
+	const TArray<FName> ObjectPaths = CurrentAssetPaths.Array();
+
+	FText ResultText;
+	bool bSuccess = false;
+	{
+		int32 NumRemoved = 0;
+		if (CollectionManagerModule.Get().RemoveFromCollection(InCollectionKey.Name, InCollectionKey.Type, ObjectPaths, &NumRemoved))
+		{
+			bSuccess = true;
+
+			FFormatNamedArguments Args;
+			Args.Add(TEXT("Number"), NumRemoved);
+			Args.Add(TEXT("CollectionName"), FText::FromName(InCollectionKey.Name));
+			ResultText = FText::Format(LOCTEXT("CollectionAssetsRemoved", "Removed {Number} asset(s) from {CollectionName}"), Args);
+		}
+		else
+		{
+			ResultText = CollectionManagerModule.Get().GetLastError();
+		}
+	}
+
+	if (!ResultText.IsEmpty())
+	{
+		FNotificationInfo Info(ResultText);
+		Info.bFireAndForget = true;
+		Info.bUseLargeFont = false;
+
+		TSharedPtr<SNotificationItem> Item = FSlateNotificationManager::Get().AddNotification(Info);
+		if (Item.IsValid())
+		{
+			Item->SetCompletionState((bSuccess) ? SNotificationItem::CS_Success : SNotificationItem::CS_Fail);
+		}
+	}
+}
+
+bool FCollectionAssetManagement::IsCollectionEnabled(FCollectionNameType InCollectionKey) const
+{
+	// Non-local collections can only be changed if we have an available source control connection
+	const bool bCollectionWritable = (InCollectionKey.Type == ECollectionShareType::CST_Local) || (ISourceControlModule::Get().IsEnabled() && ISourceControlModule::Get().GetProvider().IsAvailable());
+	return bCollectionWritable && CurrentAssetPaths.Num() > 0;
+}
+
+ECheckBoxState FCollectionAssetManagement::GetCollectionCheckState(FCollectionNameType InCollectionKey) const
+{
+	// If the collection exists in the map, then it means that the current selection contains at least one asset using that collection
+	const ECheckBoxState* FoundCheckState = AssetManagementState.Find(InCollectionKey);
+	if (FoundCheckState)
+	{
+		return *FoundCheckState;
+	}
+
+	// If the collection doesn't exist in the map, then it's assumed to be unused by the current selection (and thus unchecked)
+	return ECheckBoxState::Unchecked;
+}
+
+void FCollectionAssetManagement::UpdateAssetManagementState()
+{
+	FCollectionManagerModule& CollectionManagerModule = FCollectionManagerModule::GetModule();
+
+	AssetManagementState.Empty();
+
+	if (CurrentAssetPaths.Num() == 0)
+	{
+		return;
+	}
+
+	// The logic below is much simpler when only a single object is selected as we don't need to deal with set intersection
+	if (CurrentAssetPaths.Num() == 1)
+	{
+		TArray<FCollectionNameType> MatchedCollections;
+		CollectionManagerModule.Get().GetCollectionsContainingObject(*CurrentAssetPaths.CreateConstIterator(), MatchedCollections, ECollectionRecursionFlags::Self);
+
+		for (const FCollectionNameType& CollectionKey : MatchedCollections)
+		{
+			AssetManagementState.Add(CollectionKey, ECheckBoxState::Checked);
+		}
+
+		MatchedCollections.Reset();
+		CollectionManagerModule.Get().GetCollectionsContainingObject(*CurrentAssetPaths.CreateConstIterator(), MatchedCollections, ECollectionRecursionFlags::Children);
+
+		for (const FCollectionNameType& CollectionKey : MatchedCollections)
+		{
+			if (!AssetManagementState.Contains(CollectionKey))
+			{
+				AssetManagementState.Add(CollectionKey, ECheckBoxState::Undetermined);
+			}
+		}
+	}
+	else
+	{
+		const TArray<FName> ObjectPaths = CurrentAssetPaths.Array();
+
+		TMap<FCollectionNameType, TArray<FName>> CollectionsAndMatchedObjects;
+		CollectionManagerModule.Get().GetCollectionsContainingObjects(ObjectPaths, CollectionsAndMatchedObjects);
+
+		for (const auto& MatchedCollection : CollectionsAndMatchedObjects)
+		{
+			const FCollectionNameType& CollectionKey = MatchedCollection.Key;
+			const TArray<FName>& MatchedObjects = MatchedCollection.Value;
+
+			// Collections that contain all of the selected assets are shown as checked, collections that only contain some of the selected assets are shown as undetermined
+			AssetManagementState.Add(
+				CollectionKey, 
+				(MatchedObjects.Num() == CurrentAssetPaths.Num()) ? ECheckBoxState::Checked : ECheckBoxState::Undetermined
+				);
+		}
+	}
+}
+
+void FCollectionAssetManagement::HandleCollectionRenamed(const FCollectionNameType& OriginalCollection, const FCollectionNameType& NewCollection)
+{
+	if (AssetManagementState.Contains(OriginalCollection))
+	{
+		AssetManagementState.Add(NewCollection, AssetManagementState[OriginalCollection]);
+		AssetManagementState.Remove(OriginalCollection);
+	}
+}
+
+void FCollectionAssetManagement::HandleCollectionUpdated(const FCollectionNameType& Collection)
+{
+	// Collection has changed in an unknown way - we need to update everything to be sure
+	UpdateAssetManagementState();
+}
+
+void FCollectionAssetManagement::HandleCollectionDestroyed(const FCollectionNameType& Collection)
+{
+	AssetManagementState.Remove(Collection);
+}
+
+void FCollectionAssetManagement::HandleAssetsAddedToCollection(const FCollectionNameType& Collection, const TArray<FName>& AssetsAdded)
+{
+	// Only need to update if one of the added assets belongs to our current selection set
+	bool bNeedsUpdate = false;
+	for (const FName& AssetPath : AssetsAdded)
+	{
+		if (CurrentAssetPaths.Contains(AssetPath))
+		{
+			bNeedsUpdate = true;
+			break;
+		}
+	}
+
+	if (bNeedsUpdate)
+	{
+		UpdateAssetManagementState();
+	}
+}
+
+void FCollectionAssetManagement::HandleAssetsRemovedFromCollection(const FCollectionNameType& Collection, const TArray<FName>& AssetsRemoved)
+{
+	// Only need to update if one of the removed assets belongs to our current selection set
+	bool bNeedsUpdate = false;
+	for (const FName& AssetPath : AssetsRemoved)
+	{
+		if (CurrentAssetPaths.Contains(AssetPath))
+		{
+			bNeedsUpdate = true;
+			break;
+		}
+	}
+
+	if (bNeedsUpdate)
+	{
+		UpdateAssetManagementState();
+	}
+}
+
+#undef LOCTEXT_NAMESPACE
diff --git a/Engine/Source/Developer/AssetTools/Public/CollectionAssetManagement.h b/Engine/Source/Developer/AssetTools/Public/CollectionAssetManagement.h
new file mode 100644
--- /dev/null
+++ b/Engine/Source/Developer/AssetTools/Public/CollectionAssetManagement.h
@@ -0,0 +1,66 @@
+// Copyright Epic Games, Inc. All Rights Reserved.
+
+#pragma once
+
+#include "CoreMinimal.h"
+#include "AssetData.h"
+#include "CollectionManagerTypes.h"
+#include "Styling/SlateTypes.h"
+
+/** Handles the collection management for the given assets */
+class ASSETTOOLS_API FCollectionAssetManagement
+{
+public:
+	/** Constructor */
+	FCollectionAssetManagement();
+
+	/** Destructor */
+	~FCollectionAssetManagement();
+
+	/** Set the assets that we are currently observing and managing the collection state of */
+	void SetCurrentAssets(const TArray<FAssetData>& CurrentAssets);
+
+	/** Add the current assets to the given collection */
+	void AddCurrentAssetsToCollection(FCollectionNameType InCollectionKey);
+
+	/** Remove the current assets from the given collection */
+	void RemoveCurrentAssetsFromCollection(FCollectionNameType InCollectionKey);
+
+	/** Return whether or not the given collection should be enabled in any management UIs */
+	bool IsCollectionEnabled(FCollectionNameType InCollectionKey) const;
+
+	/** Get the check box state the given collection should use in any management UIs */
+	ECheckBoxState GetCollectionCheckState(FCollectionNameType InCollectionKey) const;
+
+private:
+	/** Update the internal state used to track the check box status for each collection */
+	void UpdateAssetManagementState();
+
+	/** Handles an on collection renamed event */
+	void HandleCollectionRenamed(const FCollectionNameType& OriginalCollection, const FCollectionNameType& NewCollection);
+
+	/** Handles an on collection updated event */
+	void HandleCollectionUpdated(const FCollectionNameType& Collection);
+
+	/** Handles an on collection destroyed event */
+	void HandleCollectionDestroyed(const FCollectionNameType& Collection);
+
+	/** Handles assets being added to a collection */
+	void HandleAssetsAddedToCollection(const FCollectionNameType& Collection, const TArray<FName>& AssetsAdded);
+
+	/** Handles assets being removed from a collection */
+	void HandleAssetsRemovedFromCollection(const FCollectionNameType& Collection, const TArray<FName>& AssetsRemoved);
+
+	/** Set of asset paths that we are currently observing and managing the collection state of */
+	TSet<FName> CurrentAssetPaths;
+
+	/** Mapping between a collection and its asset management state (based on the current assets). A missing item is assumed to be unchecked */
+	TMap<FCollectionNameType, ECheckBoxState> AssetManagementState;
+
+	/** Delegate handles */
+	FDelegateHandle OnCollectionRenamedHandle;
+	FDelegateHandle OnCollectionDestroyedHandle;
+	FDelegateHandle OnCollectionUpdatedHandle;
+	FDelegateHandle OnAssetsAddedHandle;
+	FDelegateHandle OnAssetsRemovedHandle;
+};
diff --git a/Engine/Source/Editor/ContentBrowser/ContentBrowser.Build.cs b/Engine/Source/Editor/ContentBrowser/ContentBrowser.Build.cs
--- a/Engine/Source/Editor/ContentBrowser/ContentBrowser.Build.cs
+++ b/Engine/Source/Editor/ContentBrowser/ContentBrowser.Build.cs
@@ -9,7 +9,6 @@ public class ContentBrowser : ModuleRules
 		PrivateIncludePathModuleNames.AddRange(
 			new string[] {
 				"AssetRegistry",
-				"AssetTools",
 				"CollectionManager",
 				"EditorWidgets",
 				"GameProjectGeneration",
@@ -31,6 +30,7 @@ public class ContentBrowser : ModuleRules
 				"Slate",
 				"SlateCore",
                 "EditorStyle",
+				"AssetTools",
 				"SourceControl",
 				"SourceControlWindows",
 				"WorkspaceMenuStructure",
@@ -49,7 +49,6 @@ public class ContentBrowser : ModuleRules
 			new string[] {
 				"PropertyEditor",
 				"PackagesDialog",
-				"AssetTools",
 				"CollectionManager",
 				"GameProjectGeneration",
                 "MainFrame"
diff --git a/Engine/Source/Editor/ContentBrowser/Private/CollectionAssetManagement.cpp b/Engine/Source/Editor/ContentBrowser/Private/CollectionAssetManagement.cpp
deleted file mode 100644
diff --git a/Engine/Source/Editor/ContentBrowser/Private/CollectionAssetManagement.h b/Engine/Source/Editor/ContentBrowser/Private/CollectionAssetManagement.h
deleted file mode 100644
