Registered references held by FVTConversionWorker to prevent GC issues.

diff --git a/Engine/Source/Developer/AssetTools/Private/AssetVtConversion.cpp b/Engine/Source/Developer/AssetTools/Private/AssetVtConversion.cpp
--- a/Engine/Source/Developer/AssetTools/Private/AssetVtConversion.cpp
+++ b/Engine/Source/Developer/AssetTools/Private/AssetVtConversion.cpp
@@ -120,6 +120,21 @@ bool IsVirtualTextureValidForMaterial(UMaterialInterface* InMaterial, UTexture2D
 	return true;
 }
 
+void FVTConversionWorker::AddReferencedObjects( FReferenceCollector& Collector )
+{
+	Collector.AddReferencedObjects( UserTextures );
+	Collector.AddReferencedObjects( Textures );
+	Collector.AddReferencedObjects( Materials );
+	Collector.AddReferencedObjects( Functions );
+	Collector.AddReferencedObjects( SizeRejectedTextures );
+	Collector.AddReferencedObjects( SizeRejectedTextures );
+	Collector.AddReferencedObjects( MaterialRejectedTextures );
+
+	TArray<UObject*> AuditTrailKeys;
+	AuditTrail.GetKeys(AuditTrailKeys);
+	Collector.AddReferencedObjects( AuditTrailKeys );
+}
+
 void FVTConversionWorker::FindAllTexturesAndMaterials_Iteration(TArray<UMaterial*>& InAffectedMaterials,
 	TArray<UMaterialFunctionInterface*>& InAffectedFunctions,
 	TArray<UTexture2D*>& InAffectedTextures,
diff --git a/Engine/Source/Developer/AssetTools/Private/AssetVtConversion.h b/Engine/Source/Developer/AssetTools/Private/AssetVtConversion.h
--- a/Engine/Source/Developer/AssetTools/Private/AssetVtConversion.h
+++ b/Engine/Source/Developer/AssetTools/Private/AssetVtConversion.h
@@ -4,13 +4,14 @@
 
 #include "CoreMinimal.h"
 #include "AssetRegistryModule.h"
+#include "UObject/GCObject.h"
 
 class UTexture2D;
 class UMaterial;
 class UMaterialFunctionInterface;
 struct FScopedSlowTask;
 
-class FVTConversionWorker
+class FVTConversionWorker : public FGCObject
 {
 public:
 	FVTConversionWorker(bool bInConvertBackward = false) : bConvertBackward(bInConvertBackward){}
@@ -42,6 +43,9 @@ public:
 
 	// Based on the filtered lists do the actual VT conversion
 	void DoConvert();
+
+	void AddReferencedObjects( FReferenceCollector& Collector ) override;
+
 private:
 	void FindAllTexturesAndMaterials_Iteration(TArray<UMaterial*>& InAffectedMaterials,
 		TArray<UMaterialFunctionInterface*>& InAffectedFunctions,
