Add asset file path length to content browser tooltip

Upcoming cooker changes remove the path length limitations for cooked file paths. This change forks and specializes/simplifies the path computation for source asset file paths, and adds the file path length (and maximum allowed) to the tooltip so users can still peek about any prohibitively long file paths in their project on disk.

diff --git a/Engine/Source/Developer/AssetTools/Private/AssetViewUtils.cpp b/Engine/Source/Developer/AssetTools/Private/AssetViewUtils.cpp
--- a/Engine/Source/Developer/AssetTools/Private/AssetViewUtils.cpp
+++ b/Engine/Source/Developer/AssetTools/Private/AssetViewUtils.cpp
@@ -1318,6 +1318,23 @@ bool AssetViewUtils::IsValidFolderPathForCreate(const FString& InFolderPath, con
 	return true;
 }
 
+FString AssetViewUtils::GetPackagePathWithinRoot(const FString& PackageName)
+{
+	const FString AbsoluteRootPath = FPaths::ConvertRelativePathToFull(FPaths::RootDir());
+
+	FString RelativePathToAsset;
+	if (FPackageName::TryConvertLongPackageNameToFilename(PackageName, RelativePathToAsset))
+	{
+		const FString AbsolutePathToAsset = FPaths::ConvertRelativePathToFull(RelativePathToAsset);
+
+		RelativePathToAsset = AbsolutePathToAsset;
+		FPaths::RemoveDuplicateSlashes(RelativePathToAsset);
+		RelativePathToAsset.RemoveFromStart(AbsoluteRootPath, ESearchCase::CaseSensitive);
+	}
+
+	return RelativePathToAsset;
+}
+
 int32 AssetViewUtils::GetPackageLengthForCooking(const FString& PackageName, bool IsInternalBuild)
 {
 	// We assume the game name is 20 characters (the maximum allowed) to make sure that content can be ported between projects
@@ -1435,6 +1452,22 @@ bool AssetViewUtils::IsValidPackageForCooking(const FString& PackageName, FText&
 	return true;
 }
 
+int32 AssetViewUtils::GetMaxAssetPathLen()
+{
+	const int32 PathRootAffordance = 50;
+
+	if ( GetDefault<UEditorExperimentalSettings>()->bEnableLongPathsSupport )
+	{
+		// Allow the longest path allowed by the system
+		return FPlatformMisc::GetMaxPathLength() - PathRootAffordance;
+	}
+	else
+	{
+		// 260 characters is the limit on Windows, which is the shortest max path of any platforms that support cooking
+		return 260 - PathRootAffordance;
+	}
+}
+
 int32 AssetViewUtils::GetMaxCookPathLen()
 {
 	if (GetDefault<UEditorExperimentalSettings>()->bEnableLongPathsSupport)
diff --git a/Engine/Source/Developer/AssetTools/Public/AssetViewUtils.h b/Engine/Source/Developer/AssetTools/Public/AssetViewUtils.h
--- a/Engine/Source/Developer/AssetTools/Public/AssetViewUtils.h
+++ b/Engine/Source/Developer/AssetTools/Public/AssetViewUtils.h
@@ -177,12 +177,18 @@ namespace AssetViewUtils
 	/** Returns true if the specified folder name in the specified path is available for folder creation */
 	ASSETTOOLS_API bool IsValidFolderPathForCreate(const FString& FolderPath, const FString& NewFolderName, FText& OutErrorMessage);
 
+	/** Returns the relative path, from the workspace root, of the package */
+	ASSETTOOLS_API FString GetPackagePathWithinRoot(const FString& PackageName);
+
 	/** Returns the length of the computed cooked package name and path whether it's run on a build machine or locally */
 	ASSETTOOLS_API int32 GetPackageLengthForCooking(const FString& PackageName, bool IsInternalBuild);
 
 	/** Checks to see whether the path is within the size restrictions for cooking */
 	ASSETTOOLS_API bool IsValidPackageForCooking(const FString& PackageName, FText& OutErrorMessage);
 
+	/** Gets the maximum path length for an asset package file. Changes behavior based on whether the editor experimental setting for long paths is enabled. */
+	ASSETTOOLS_API int32 GetMaxAssetPathLen();
+
 	/** Gets the maximum path length for a cooked file. Changes behavior based on whether the editor experimental setting for long paths is enabled. */
 	ASSETTOOLS_API int32 GetMaxCookPathLen();
 
diff --git a/Engine/Source/Editor/ContentBrowser/Private/AssetViewWidgets.cpp b/Engine/Source/Editor/ContentBrowser/Private/AssetViewWidgets.cpp
--- a/Engine/Source/Editor/ContentBrowser/Private/AssetViewWidgets.cpp
+++ b/Engine/Source/Editor/ContentBrowser/Private/AssetViewWidgets.cpp
@@ -728,6 +728,13 @@ TSharedRef<SWidget> SAssetViewItem::CreateToolTipWidget() const
 
 			if (ItemAssetData.IsValid() && ItemAssetData.PackageName != NAME_None)
 			{
+				const FString PackagePathWithinRoot = ContentBrowserUtils::GetPackagePathWithinRoot(ItemAssetData.PackageName.ToString());
+				int32 PackageNameLength = PackagePathWithinRoot.Len();
+
+				int32 MaxAssetPathLen = ContentBrowserUtils::GetMaxAssetPathLen();
+				AddToToolTipInfoBox( InfoBox, LOCTEXT( "TileViewTooltipAssetPathLengthKey", "Asset Filepath Length" ), FText::Format( LOCTEXT( "TileViewTooltipAssetPathLengthValue", "{0} / {1}" ),
+					FText::AsNumber( PackageNameLength ), FText::AsNumber( MaxAssetPathLen ) ), PackageNameLength > MaxAssetPathLen ? true : false );
+
 				int32 PackageNameLengthForCooking = ContentBrowserUtils::GetPackageLengthForCooking(ItemAssetData.PackageName.ToString(), FEngineBuildSettings::IsInternalBuild());
 
 				int32 MaxCookPathLen = ContentBrowserUtils::GetMaxCookPathLen();
