During VT texture conversion, call UTexture::Modify() before making changes rather than after
This is correct usage, important now due to async texture changes

diff --git a/Engine/Source/Developer/AssetTools/Private/AssetVtConversion.cpp b/Engine/Source/Developer/AssetTools/Private/AssetVtConversion.cpp
--- a/Engine/Source/Developer/AssetTools/Private/AssetVtConversion.cpp
+++ b/Engine/Source/Developer/AssetTools/Private/AssetVtConversion.cpp
@@ -597,12 +597,11 @@ void FVTConversionWorker::DoConvert()
 				UE_LOG(LogVirtualTextureConversion, Display, TEXT("Texture %s"), *TextureToUpdate->GetName());
 				TextureTask.EnterProgressFrame();
 
-				bool OldVt = TextureToUpdate->VirtualTextureStreaming;
-				TextureToUpdate->VirtualTextureStreaming = bVirtualTextureEnable;
-				if (OldVt != TextureToUpdate->VirtualTextureStreaming)
+				if (TextureToUpdate->VirtualTextureStreaming != bVirtualTextureEnable)
 				{
 					FPropertyChangedEvent PropertyChangeEvent(UTexture::StaticClass()->FindPropertyByName(GET_MEMBER_NAME_CHECKED(UTexture, VirtualTextureStreaming)));
 					TextureToUpdate->Modify();
+					TextureToUpdate->VirtualTextureStreaming = bVirtualTextureEnable;
 					TextureToUpdate->PostEditChangeProperty(PropertyChangeEvent);
 				}
 			}
