Implemented support for right-click Nanite submenu on static mesh(es), with support for batch enable/disable functionality.

diff --git a/Engine/Source/Developer/AssetTools/Private/AssetTypeActions/AssetTypeActions_StaticMesh.cpp b/Engine/Source/Developer/AssetTools/Private/AssetTypeActions/AssetTypeActions_StaticMesh.cpp
--- a/Engine/Source/Developer/AssetTools/Private/AssetTypeActions/AssetTypeActions_StaticMesh.cpp
+++ b/Engine/Source/Developer/AssetTools/Private/AssetTypeActions/AssetTypeActions_StaticMesh.cpp
@@ -42,6 +42,15 @@ void FAssetTypeActions_StaticMesh::GetActions(const TArray<UObject*>& InObjects,
 			);
 	}
 
+	Section.AddSubMenu(
+		"StaticMesh_NaniteMenu",
+		NSLOCTEXT("AssetTypeActions_StaticMesh", "StaticMesh_NaniteMenu", "Nanite"),
+		NSLOCTEXT("AssetTypeActions_StaticMesh", "StaticMesh_NaniteTooltip", "Nanite Options and Tools"),
+		FNewMenuDelegate::CreateSP(this, &FAssetTypeActions_StaticMesh::GetNaniteMenu, Meshes),
+		false,
+		FSlateIcon(FEditorStyle::GetStyleSetName(), "ContentBrowser.AssetActions")
+	);
+
 	Section.AddSubMenu(
 		"StaticMesh_LODMenu",
 		NSLOCTEXT("AssetTypeActions_StaticMesh", "StaticMesh_LODMenu", "Level Of Detail"),
@@ -102,6 +111,23 @@ void FAssetTypeActions_StaticMesh::GetResolvedSourceFilePaths(const TArray<UObje
 	}
 }
 
+void FAssetTypeActions_StaticMesh::GetNaniteMenu(class FMenuBuilder& MenuBuilder, TArray<TWeakObjectPtr<UStaticMesh>> Meshes)
+{
+	MenuBuilder.AddMenuEntry(
+		NSLOCTEXT("AssetTypeActions_StaticMesh", "StaticMesh_NaniteEnable", "Enable"),
+		NSLOCTEXT("AssetTypeActions_StaticMesh", "StaticMesh_NaniteEnableTooltip", "Enables support for Nanite on the selected mesh(es)."),
+		FSlateIcon(),
+		FUIAction(FExecuteAction::CreateSP(this, &FAssetTypeActions_StaticMesh::ExecuteNaniteEnable, Meshes))
+	);
+
+	MenuBuilder.AddMenuEntry(
+		NSLOCTEXT("AssetTypeActions_StaticMesh", "StaticMesh_NaniteDisable", "Disable"),
+		NSLOCTEXT("AssetTypeActions_StaticMesh", "StaticMesh_NaniteDisableTooltip", "Disables support for Nanite on the selected mesh(es)."),
+		FSlateIcon(),
+		FUIAction(FExecuteAction::CreateSP(this, &FAssetTypeActions_StaticMesh::ExecuteNaniteDisable, Meshes))
+	);
+}
+
 void FAssetTypeActions_StaticMesh::GetImportLODMenu(class FMenuBuilder& MenuBuilder,TArray<TWeakObjectPtr<UStaticMesh>> Objects)
 {
 	check(Objects.Num() > 0);
@@ -126,7 +152,6 @@ void FAssetTypeActions_StaticMesh::GetImportLODMenu(class FMenuBuilder& MenuBuil
 
 void FAssetTypeActions_StaticMesh::GetLODMenu(class FMenuBuilder& MenuBuilder, TArray<TWeakObjectPtr<UStaticMesh>> Meshes)
 {
-
 	MenuBuilder.AddSubMenu(
 		NSLOCTEXT("AssetTypeActions_StaticMesh", "StaticMesh_ImportLOD", "Import LOD"),
 		NSLOCTEXT("AssetTypeActions_StaticMesh", "StaticMesh_ImportLODtooltip", "Imports meshes into the LODs"),
@@ -159,9 +184,7 @@ void FAssetTypeActions_StaticMesh::GetLODMenu(class FMenuBuilder& MenuBuilder, T
 		FExecuteAction::CreateSP(this, &FAssetTypeActions_StaticMesh::ExecutePasteLODSettings, Meshes),
 		FCanExecuteAction::CreateSP(this, &FAssetTypeActions_StaticMesh::CanPasteLODSettings, Meshes)
 		)
-		);
-
-
+	);
 }
 
 void FAssetTypeActions_StaticMesh::ExecuteImportMeshLOD(UObject* Mesh, int32 LOD)
@@ -174,7 +197,6 @@ void FAssetTypeActions_StaticMesh::ExecuteCopyLODSettings(TArray<TWeakObjectPtr<
 	LODCopyMesh = Objects[0];
 }
 
-
 bool FAssetTypeActions_StaticMesh::CanCopyLODSettings(TArray<TWeakObjectPtr<UStaticMesh>> Objects) const
 {
 	return Objects.Num() == 1;
@@ -265,4 +287,34 @@ void FAssetTypeActions_StaticMesh::ExecuteRemoveVertexColors(TArray<TWeakObjectP
 	}
 }
 
+void FAssetTypeActions_StaticMesh::ExecuteNaniteEnable(TArray<TWeakObjectPtr<UStaticMesh>> Objects)
+{
+	for (auto StaticMeshPtr : Objects)
+	{
+		UStaticMesh* Mesh = StaticMeshPtr.Get();
+		if (Mesh)
+		{
+			Mesh->NaniteSettings.bEnabled = true;
+			Mesh->Build();
+			Mesh->MarkPackageDirty();
+			Mesh->OnMeshChanged.Broadcast();
+		}
+	}
+}
+
+void FAssetTypeActions_StaticMesh::ExecuteNaniteDisable(TArray<TWeakObjectPtr<UStaticMesh>> Objects)
+{
+	for (auto StaticMeshPtr : Objects)
+	{
+		UStaticMesh* Mesh = StaticMeshPtr.Get();
+		if (Mesh)
+		{
+			Mesh->NaniteSettings.bEnabled = false;
+			Mesh->Build();
+			Mesh->MarkPackageDirty();
+			Mesh->OnMeshChanged.Broadcast();
+		}
+	}
+}
+
 #undef LOCTEXT_NAMESPACE
diff --git a/Engine/Source/Developer/AssetTools/Private/AssetTypeActions/AssetTypeActions_StaticMesh.h b/Engine/Source/Developer/AssetTools/Private/AssetTypeActions/AssetTypeActions_StaticMesh.h
--- a/Engine/Source/Developer/AssetTools/Private/AssetTypeActions/AssetTypeActions_StaticMesh.h
+++ b/Engine/Source/Developer/AssetTools/Private/AssetTypeActions/AssetTypeActions_StaticMesh.h
@@ -46,6 +46,15 @@ private:
 	/** Handler for when RemoveVertexColors is selected */
 	void ExecuteRemoveVertexColors(TArray<TWeakObjectPtr<UStaticMesh>> Objects);
 
+	/** Handler for when NaniteEnable is selected */
+	void ExecuteNaniteEnable(TArray<TWeakObjectPtr<UStaticMesh>> Objects);
+
+	/** Handler for when NaniteDisable is selected */
+	void ExecuteNaniteDisable(TArray<TWeakObjectPtr<UStaticMesh>> Objects);
+
+	/** Handler to provide the Nanite sub-menu */
+	void GetNaniteMenu(class FMenuBuilder& MenuBuilder, TArray<TWeakObjectPtr<UStaticMesh>> Meshes);
+
 	/** Handler to provide the list of LODs that can be imported or reimported */
 	void GetImportLODMenu(class FMenuBuilder& MenuBuilder,TArray<TWeakObjectPtr<UStaticMesh>> Objects);
 
