Moving Cascade particle systems to Misc category

diff --git a/Engine/Source/Developer/AssetTools/Private/AssetTypeActions/AssetTypeActions_ParticleSystem.h b/Engine/Source/Developer/AssetTools/Private/AssetTypeActions/AssetTypeActions_ParticleSystem.h
--- a/Engine/Source/Developer/AssetTools/Private/AssetTypeActions/AssetTypeActions_ParticleSystem.h
+++ b/Engine/Source/Developer/AssetTools/Private/AssetTypeActions/AssetTypeActions_ParticleSystem.h
@@ -12,13 +12,13 @@ class FAssetTypeActions_ParticleSystem : public FAssetTypeActions_Base
 {
 public:
 	// IAssetTypeActions Implementation
-	virtual FText GetName() const override { return NSLOCTEXT("AssetTypeActions", "AssetTypeActions_ParticleSystem", "Particle System"); }
+	virtual FText GetName() const override { return NSLOCTEXT("AssetTypeActions", "AssetTypeActions_ParticleSystem", "Cascade Particle System (Legacy)"); }
 	virtual FColor GetTypeColor() const override { return FColor(255,255,255); }
 	virtual UClass* GetSupportedClass() const override { return UParticleSystem::StaticClass(); }
 	virtual bool HasActions ( const TArray<UObject*>& InObjects ) const override { return true; }
 	virtual void GetActions(const TArray<UObject*>& InObjects, struct FToolMenuSection& Section) override;
 	virtual void OpenAssetEditor( const TArray<UObject*>& InObjects, TSharedPtr<class IToolkitHost> EditWithinLevelEditor = TSharedPtr<IToolkitHost>() ) override;
-	virtual uint32 GetCategories() override { return EAssetTypeCategories::Basic; }
+	virtual uint32 GetCategories() override { return EAssetTypeCategories::Misc; }
 	
 private:
 	/** Handler for when Copy Parameters is selected */
