Fixed issue where renaming certain assets causes soft reference issues, with special handling for GameMapsSettings: ensure redirectors are created for assets that are soft-referenced by CDOs when they are renamed. Enhanced warning message at time of renaming if any GameMapsSettings are affected. Added error at cook time if any of the GameMapsSettings maps/modes are still pointing to redirectors, as they will fail to load at runtime.

diff --git a/Engine/Source/Developer/AssetTools/AssetTools.Build.cs b/Engine/Source/Developer/AssetTools/AssetTools.Build.cs
--- a/Engine/Source/Developer/AssetTools/AssetTools.Build.cs
+++ b/Engine/Source/Developer/AssetTools/AssetTools.Build.cs
@@ -38,6 +38,7 @@ public class AssetTools : ModuleRules
 				"ToolMenus",
 				"PhysicsCore",
 				"DeveloperSettings",
+				"EngineSettings",
 				"InterchangeCore",
 				"InterchangeEngine"
 			}
diff --git a/Engine/Source/Developer/AssetTools/Private/AssetRenameManager.cpp b/Engine/Source/Developer/AssetTools/Private/AssetRenameManager.cpp
--- a/Engine/Source/Developer/AssetTools/Private/AssetRenameManager.cpp
+++ b/Engine/Source/Developer/AssetTools/Private/AssetRenameManager.cpp
@@ -47,6 +47,7 @@
 #include "Settings/EditorProjectSettings.h"
 #include "Engine/World.h"
 #include "Engine/MapBuildDataRegistry.h"
+#include "GameMapsSettings.h"
 
 #define LOCTEXT_NAMESPACE "AssetRenameManager"
 
@@ -90,11 +91,13 @@
 	FText FailureReason;
 	bool bCreateRedirector;
 	bool bRenameFailed;
+	bool bWarnAboutProjectSettingsReference;
 
 	FAssetRenameDataWithReferencers(const FAssetRenameData& InRenameData)
 		: FAssetRenameData(InRenameData)
 		, bCreateRedirector(false)
 		, bRenameFailed(false)
+		, bWarnAboutProjectSettingsReference(false)
 	{
 		if (Asset.IsValid() && !OldObjectPath.IsValid())
 		{
@@ -364,22 +367,43 @@
 	}
 
 	// Warn the user if they are about to rename an asset that is referenced by a CDO
-	TArray<TWeakObjectPtr<UObject>> CDOAssets = FindCDOReferencedAssets(AssetsToRename);
+	TArray<FAssetRenameDataWithReferencers*> CDOHardReferencedAssets, CDOSoftReferenceRenames;
+	FindCDOReferences(AssetsToRename, CDOHardReferencedAssets, CDOSoftReferenceRenames, true);
 
 	// Warn the user if there were any references
-	if (CDOAssets.Num())
+	if (CDOHardReferencedAssets.Num() || CDOSoftReferenceRenames.Num())
 	{
 		FString AssetNames;
-		for (auto AssetIt = CDOAssets.CreateConstIterator(); AssetIt; ++AssetIt)
+		for (auto HardAssetIt = CDOHardReferencedAssets.CreateConstIterator(); HardAssetIt; ++HardAssetIt)
 		{
-			UObject* Asset = (*AssetIt).Get();
+			UObject* Asset = (*HardAssetIt)->Asset.Get();
 			if (Asset)
 			{
 				AssetNames += FString("\n") + Asset->GetName();
 			}
 		}
 
-		const FText MessageText = FText::Format(LOCTEXT("RenameCDOReferences", "The following assets are referenced by one or more Class Default Objects: \n{0}\n\nContinuing with the rename may require code changes to fix these references. Do you wish to continue?"), FText::FromString(AssetNames));
+		for (auto SoftRefIt = CDOSoftReferenceRenames.CreateConstIterator(); SoftRefIt; ++SoftRefIt)
+		{
+			UObject* Asset = (*SoftRefIt)->Asset.Get();
+			if (Asset)
+			{
+				FString OptionalTagsString;
+
+				if ((*SoftRefIt)->bWarnAboutProjectSettingsReference)
+				{
+					OptionalTagsString = LOCTEXT("ProjSettingsReferenceTag", "project settings soft reference").ToString();
+				}
+				else
+				{
+					OptionalTagsString = LOCTEXT("SoftReferenceTag", "soft reference").ToString();
+				}
+				
+				AssetNames += FString::Printf(TEXT("\n%s (%s)"), *Asset->GetName(), *OptionalTagsString);
+			}
+		}
+
+		const FText MessageText = FText::Format(LOCTEXT("RenameCDOReferences", "The following assets are referenced by one or more Class Default Objects: \n{0}\n\nContinuing with the rename may require changes to fix references in code or project settings. Assets could otherwise be missing in cooked/packaged builds. Do you wish to continue?"), FText::FromString(AssetNames));
 		if (FMessageDialog::Open(EAppMsgType::YesNo, EAppReturnType::No, MessageText) == EAppReturnType::No)
 		{
 			return false;
@@ -463,14 +487,170 @@
 	return ReportFailures(AssetsToRename, bWithDialog) == 0;
 }
 
-TArray<TWeakObjectPtr<UObject>> FAssetRenameManager::FindCDOReferencedAssets(const TArray<FAssetRenameDataWithReferencers>& AssetsToRename) const
+
+struct FSoftObjectPathRenameSerializer : public FArchiveUObject
 {
-	TArray<TWeakObjectPtr<UObject>> CDOAssets, LocalAssetsToRename;
-	for (const FAssetRenameDataWithReferencers& AssetToRename : AssetsToRename)
+	void StartSerializingObject(UObject* InCurrentObject)
 	{
+		CurrentObject = InCurrentObject;
+		bFoundReference = false;
+	}
+	bool HasFoundReference() const
+	{
+		return bFoundReference;
+	}
+
+	FSoftObjectPathRenameSerializer(const TMap<FSoftObjectPath, FSoftObjectPath>& InRedirectorMap,
+		bool bInCheckOnly,
+		TMap<FSoftObjectPath, TSet<FWeakObjectPtr>>* InCachedObjectPaths,
+		const FName InPackageName = NAME_None)
+		: RedirectorMap(InRedirectorMap)
+		, CachedObjectPaths(InCachedObjectPaths)
+		, CurrentObject(nullptr)
+		, PackageName(InPackageName)
+		, bSearchOnly(bInCheckOnly)
+		, bFoundReference(false)
+	{
+		if (InCachedObjectPaths)
+		{
+			DirtyDelegateHandle = UPackage::PackageMarkedDirtyEvent.AddRaw(this, &FSoftObjectPathRenameSerializer::OnMarkPackageDirty);
+		}
+
+		this->ArIsObjectReferenceCollector = true;
+		this->ArIsModifyingWeakAndStrongReferences = true;
+
+		// Mark it as saving to correctly process all references
+		this->SetIsSaving(true);
+	}
+
+	virtual ~FSoftObjectPathRenameSerializer()
+	{
+		UPackage::PackageMarkedDirtyEvent.Remove(DirtyDelegateHandle);
+	}
+
+	virtual bool ShouldSkipProperty(const FProperty* InProperty) const override
+	{
+		if (InProperty->HasAnyPropertyFlags(CPF_Transient | CPF_Deprecated | CPF_IsPlainOldData))
+		{
+			return true;
+		}
+
+		FFieldClass* PropertyClass = InProperty->GetClass();
+		if (PropertyClass->GetCastFlags() & (CASTCLASS_FBoolProperty | CASTCLASS_FNameProperty | CASTCLASS_FStrProperty | CASTCLASS_FMulticastDelegateProperty))
+		{
+			return true;
+		}
+
+		if (PropertyClass->GetCastFlags() & (CASTCLASS_FArrayProperty | CASTCLASS_FMapProperty | CASTCLASS_FSetProperty))
+		{
+			if (const FArrayProperty* ArrayProperty = CastField<FArrayProperty>(InProperty))
+			{
+				return ShouldSkipProperty(ArrayProperty->Inner);
+			}
+			else if (const FMapProperty* MapProperty = CastField<FMapProperty>(InProperty))
+			{
+				return ShouldSkipProperty(MapProperty->KeyProp) && ShouldSkipProperty(MapProperty->ValueProp);
+			}
+			else if (const FSetProperty* SetProperty = CastField<FSetProperty>(InProperty))
+			{
+				return ShouldSkipProperty(SetProperty->ElementProp);
+			}
+		}
+
+		return false;
+	}
+
+	FArchive& operator<<(FSoftObjectPath& Value)
+	{
+		using namespace AssetRenameManagerImpl;
+
+		// Ignore untracked references if just doing a search only. We still want to fix them up if they happen to be there
+		if (bSearchOnly)
+		{
+			FSoftObjectPathThreadContext& ThreadContext = FSoftObjectPathThreadContext::Get();
+			FName ReferencingPackageName, ReferencingPropertyName;
+			ESoftObjectPathCollectType CollectType = ESoftObjectPathCollectType::AlwaysCollect;
+			ESoftObjectPathSerializeType SerializeType = ESoftObjectPathSerializeType::AlwaysSerialize;
+
+			ThreadContext.GetSerializationOptions(ReferencingPackageName, ReferencingPropertyName, CollectType, SerializeType, this);
+
+			if (CollectType == ESoftObjectPathCollectType::NeverCollect)
+			{
+				return *this;
+			}
+		}
+
+		if (CachedObjectPaths)
+		{
+			TSet<FWeakObjectPtr>* ObjectSet = &CachedObjectPaths->FindOrAdd(Value);
+			ObjectSet->Add(CurrentObject);
+		}
+
+		const FString& SubPath = Value.GetSubPathString();
+		for (const TPair<FSoftObjectPath, FSoftObjectPath>& Pair : RedirectorMap)
+		{
+			if (Pair.Key.GetAssetPathName() == Value.GetAssetPathName())
+			{
+				// Same asset, fix sub path. Asset will be fixed by normal serializePath call below
+				const FString& CheckSubPath = Pair.Key.GetSubPathString();
+
+				if (IsSubPath(SubPath, CheckSubPath))
+				{
+					bFoundReference = true;
+
+					if (!bSearchOnly)
+					{
+						if (CurrentObject)
+						{
+							check(!CachedObjectPaths); // Modify can invalidate the object paths map, not allowed to be modifying and using the cache at the same time
+							CurrentObject->Modify(true);
+						}
+
+						FString NewSubPath(SubPath);
+						NewSubPath.ReplaceInline(*CheckSubPath, *Pair.Value.GetSubPathString());
+						Value = FSoftObjectPath(Pair.Value.GetAssetPathName(), NewSubPath);
+					}
+					break;
+				}
+			}
+		}
+
+		return *this;
+	}
+
+	void OnMarkPackageDirty(UPackage* Pkg, bool bWasDirty)
+	{
+		UPackage::PackageMarkedDirtyEvent.Remove(DirtyDelegateHandle);
+
+		if (CachedObjectPaths && Pkg && Pkg->GetFName() == PackageName)
+		{
+			UE_LOG(LogAssetTools, VeryVerbose, TEXT("Performance: Package unexpectedly modified during serialization by FSoftObjectPathRenameSerializer: %s"), *Pkg->GetFullName());
+		}
+	}
+
+private:
+	const TMap<FSoftObjectPath, FSoftObjectPath>& RedirectorMap;
+	TMap<FSoftObjectPath, TSet<FWeakObjectPtr>>* CachedObjectPaths;
+	FDelegateHandle DirtyDelegateHandle;
+	UObject* CurrentObject;
+	FName PackageName;
+	bool bSearchOnly;
+	bool bFoundReference;
+
+};
+
+void FAssetRenameManager::FindCDOReferences(const TArrayView<FAssetRenameDataWithReferencers>& AssetsToRename, TArray<FAssetRenameDataWithReferencers*>& OutHardReferences, TArray<FAssetRenameDataWithReferencers*>& OutSoftReferences, bool bSetRedirectorFlags) const
+{
+	// Checking reference candidates off as we find them to reduce workload
+	TArray<FAssetRenameDataWithReferencers*> RemainingHardRefAssetChecklist;
+	TArray<FAssetRenameDataWithReferencers*> RemainingSoftRefAssetChecklist;
+
+	for (FAssetRenameDataWithReferencers& AssetToRename : AssetsToRename)
+	{
 		if (AssetToRename.Asset.IsValid())
 		{
-			LocalAssetsToRename.Push(AssetToRename.Asset);
+			RemainingHardRefAssetChecklist.Push(&AssetToRename);
+			RemainingSoftRefAssetChecklist.Push(&AssetToRename);
 		}
 	}
 
@@ -479,8 +659,9 @@
 	{
 		UClass* Cls = (*ClassDefaultObjectIt);
 		UObject* CDO = Cls->ClassDefaultObject;
+		bool bWorkingOnGameMapsSettings = (CDO == GetDefault<UGameMapsSettings>());
 
-		if (!CDO || !CDO->HasAllFlags(RF_ClassDefaultObject) || Cls->ClassGeneratedBy != nullptr)
+		if (!CDO || !CDO->HasAllFlags(RF_ClassDefaultObject) || CDO->IsPendingKill() || Cls->ClassGeneratedBy != nullptr)
 		{
 			continue;
 		}
@@ -491,31 +672,74 @@
 			continue;
 		}
 
-		for (TFieldIterator<FObjectProperty> PropertyIt(Cls); PropertyIt; ++PropertyIt)
+		// Search this CDO for hard references
+		for (TFieldIterator<FObjectProperty> PropertyIt(Cls); PropertyIt && RemainingHardRefAssetChecklist.Num(); ++PropertyIt)
 		{
 			const UObject* Object = PropertyIt->GetPropertyValue(PropertyIt->ContainerPtrToValuePtr<UObject>(CDO));
-			for (const TWeakObjectPtr<UObject>& Asset : LocalAssetsToRename)
+
+			for (FAssetRenameDataWithReferencers* AssetToRename : RemainingHardRefAssetChecklist)
 			{
-				if (Object == Asset.Get())
+				if (Object == AssetToRename->Asset.Get())
 				{
-					CDOAssets.Push(Asset);
-					LocalAssetsToRename.Remove(Asset);
+					OutHardReferences.Push(AssetToRename);
+					RemainingHardRefAssetChecklist.Remove(AssetToRename);
+					break;
+				}
+			}
+		}
 
-					if (LocalAssetsToRename.Num() == 0)
+		//Search this CDO for soft references
+		TMap<FSoftObjectPath, FSoftObjectPath> DummyEmptyRedirectorMap;
+		TMap<FSoftObjectPath, TSet<FWeakObjectPtr>> SoftReferenceMap;
+		FSoftObjectPathRenameSerializer SoftRefCheckSerializer(DummyEmptyRedirectorMap, true, &SoftReferenceMap);
+
+		// Gather all soft references
+		SoftRefCheckSerializer.StartSerializingObject(CDO);
+		CDO->Serialize(SoftRefCheckSerializer);
+
+		// Check all soft references in the CDO for matching items that are to be renamed, with special handling for UBlueprint assets
+		for (auto Iter = SoftReferenceMap.CreateIterator(); Iter && RemainingSoftRefAssetChecklist.Num(); ++Iter)
+		{
+			FSoftObjectPath SoftRefObjPath = Iter.Key();
+
+			if (SoftRefObjPath.IsValid())
+			{
+				TWeakObjectPtr<UObject> Object = SoftRefObjPath.ResolveObject();
+				TWeakObjectPtr<UBlueprint> ObjectAsBP = UBlueprint::GetBlueprintFromClass(Cast<UBlueprintGeneratedClass>(Object));
+
+				// Resolve to the redirected asset path if necessary
+				FName FinalSoftObjPathName = SoftRefObjPath.GetAssetPathName();
+			
+				if (!Object.IsValid() && SoftRefObjPath.IsValid())
+				{
+					FName RedirObjectPathName = GRedirectCollector.GetAssetPathRedirection(SoftRefObjPath.GetAssetPathName());
+
+					if (RedirObjectPathName != NAME_None && RedirObjectPathName != FinalSoftObjPathName)
 					{
-						// No more assets to check
-						return MoveTemp(CDOAssets);
+						FinalSoftObjPathName = RedirObjectPathName;
 					}
-					else
+				}
+
+				// Check for any matching rename requests
+				for (FAssetRenameDataWithReferencers* AssetToRename : RemainingSoftRefAssetChecklist)
+				{
+					// Look for loaded references, indirect blueprint refs to their generated class counterparts, or path name matching
+					if ((Object == AssetToRename->Asset) ||
+						(ObjectAsBP != nullptr && ObjectAsBP == AssetToRename->Asset) ||
+						(!FinalSoftObjPathName.IsNone() && FinalSoftObjPathName == AssetToRename->OldObjectPath.GetAssetPathName()))
 					{
+						AssetToRename->bCreateRedirector |= bSetRedirectorFlags;
+						AssetToRename->bWarnAboutProjectSettingsReference |= bWorkingOnGameMapsSettings;
+
+						OutSoftReferences.Push(AssetToRename);
+						RemainingSoftRefAssetChecklist.Remove(AssetToRename);
+
 						break;
 					}
 				}
 			}
 		}
 	}
-
-	return MoveTemp(CDOAssets);
 }
 
 void FAssetRenameManager::PopulateAssetReferencers(TArray<FAssetRenameDataWithReferencers>& AssetsToPopulate) const
@@ -994,158 +1218,6 @@
 	}
 }
 
-struct FSoftObjectPathRenameSerializer : public FArchiveUObject
-{
-	void StartSerializingObject(UObject* InCurrentObject)
-	{ 
-		CurrentObject = InCurrentObject;
-		bFoundReference = false; 
-	}
-	bool HasFoundReference() const 
-	{ 
-		return bFoundReference; 
-	}
-
-	FSoftObjectPathRenameSerializer(const TMap<FSoftObjectPath, FSoftObjectPath>& InRedirectorMap, bool bInCheckOnly, TMap<FSoftObjectPath, TSet<FWeakObjectPtr>>* InCachedObjectPaths, const FName InPackageName = NAME_None)
-		: RedirectorMap(InRedirectorMap)
-		, CachedObjectPaths(InCachedObjectPaths)
-		, CurrentObject(nullptr)
-		, PackageName(InPackageName)
-		, bSearchOnly(bInCheckOnly)
-		, bFoundReference(false)
-	{
-		if (InCachedObjectPaths)
-		{
-			DirtyDelegateHandle = UPackage::PackageMarkedDirtyEvent.AddRaw(this, &FSoftObjectPathRenameSerializer::OnMarkPackageDirty);
-		}
-
-		this->ArIsObjectReferenceCollector = true;
-		this->ArIsModifyingWeakAndStrongReferences = true;
-
-		// Mark it as saving to correctly process all references
-		this->SetIsSaving(true);
-	}
-
-	virtual ~FSoftObjectPathRenameSerializer()
-	{
-		UPackage::PackageMarkedDirtyEvent.Remove(DirtyDelegateHandle);
-	}
-
-	virtual bool ShouldSkipProperty(const FProperty* InProperty) const override
-	{
-		if (InProperty->HasAnyPropertyFlags(CPF_Transient | CPF_Deprecated | CPF_IsPlainOldData))
-		{
-			return true;
-		}
-
-		FFieldClass* PropertyClass = InProperty->GetClass();
-		if (PropertyClass->GetCastFlags() & (CASTCLASS_FBoolProperty | CASTCLASS_FNameProperty | CASTCLASS_FStrProperty | CASTCLASS_FMulticastDelegateProperty))
-		{
-			return true;
-		}
-
-		if (PropertyClass->GetCastFlags() & (CASTCLASS_FArrayProperty | CASTCLASS_FMapProperty | CASTCLASS_FSetProperty))
-		{
-			if (const FArrayProperty* ArrayProperty = CastField<FArrayProperty>(InProperty))
-			{
-				return ShouldSkipProperty(ArrayProperty->Inner);
-			}
-			else if (const FMapProperty* MapProperty = CastField<FMapProperty>(InProperty))
-			{
-				return ShouldSkipProperty(MapProperty->KeyProp) && ShouldSkipProperty(MapProperty->ValueProp);
-			}
-			else if (const FSetProperty* SetProperty = CastField<FSetProperty>(InProperty))
-			{
-				return ShouldSkipProperty(SetProperty->ElementProp);
-			}
-		}
-
-		return false;
-	}
-
-	FArchive& operator<<(FSoftObjectPath& Value)
-	{
-		using namespace AssetRenameManagerImpl;
-
-		// Ignore untracked references if just doing a search only. We still want to fix them up if they happen to be there
-		if (bSearchOnly)
-		{
-			FSoftObjectPathThreadContext& ThreadContext = FSoftObjectPathThreadContext::Get();
-			FName ReferencingPackageName, ReferencingPropertyName;
-			ESoftObjectPathCollectType CollectType = ESoftObjectPathCollectType::AlwaysCollect;
-			ESoftObjectPathSerializeType SerializeType = ESoftObjectPathSerializeType::AlwaysSerialize;
-
-			ThreadContext.GetSerializationOptions(ReferencingPackageName, ReferencingPropertyName, CollectType, SerializeType, this);
-
-			if (CollectType == ESoftObjectPathCollectType::NeverCollect)
-			{
-				return *this;
-			}
-		}
-
-		if (CachedObjectPaths)
-		{
-			TSet<FWeakObjectPtr>* ObjectSet = CachedObjectPaths->Find(Value);
-			if (ObjectSet == nullptr)
-			{
-				ObjectSet = &CachedObjectPaths->Add(Value);
-			}
-			ObjectSet->Add(CurrentObject);
-		}
-
-		const FString& SubPath = Value.GetSubPathString();
-		for (const TPair<FSoftObjectPath, FSoftObjectPath>& Pair : RedirectorMap)
-		{
-			if (Pair.Key.GetAssetPathName() == Value.GetAssetPathName())
-			{
-				// Same asset, fix sub path. Asset will be fixed by normal serializePath call below
-				const FString& CheckSubPath = Pair.Key.GetSubPathString();
-
-				if (IsSubPath(SubPath, CheckSubPath))
-				{
-					bFoundReference = true;
-
-					if (!bSearchOnly)
-					{
-						if (CurrentObject)
-						{
-							check(!CachedObjectPaths); // Modify can invalidate the object paths map, not allowed to be modifying and using the cache at the same time
-							CurrentObject->Modify(true);
-						}
-
-						FString NewSubPath(SubPath);
-						NewSubPath.ReplaceInline(*CheckSubPath, *Pair.Value.GetSubPathString());
-						Value = FSoftObjectPath(Pair.Value.GetAssetPathName(), NewSubPath);
-					}
-					break;
-				}
-			}
-		}
-
-		return *this;
-	}
-
-	void OnMarkPackageDirty(UPackage* Pkg, bool bWasDirty)
-	{
-		UPackage::PackageMarkedDirtyEvent.Remove(DirtyDelegateHandle);
-
-		if (CachedObjectPaths && Pkg && Pkg->GetFName() == PackageName)
-		{
-			UE_LOG(LogAssetTools, VeryVerbose, TEXT("Performance: Package unexpectedly modified during serialization by FSoftObjectPathRenameSerializer: %s"), *Pkg->GetFullName());
-		}
-	}
-
-private:
-	const TMap<FSoftObjectPath, FSoftObjectPath>& RedirectorMap;
-	TMap<FSoftObjectPath, TSet<FWeakObjectPtr>>* CachedObjectPaths;
-	FDelegateHandle DirtyDelegateHandle;
-	UObject* CurrentObject;
-	FName PackageName;
-	bool bSearchOnly;
-	bool bFoundReference;
-
-};
-
 void FAssetRenameManager::RenameReferencingSoftObjectPaths(const TArray<UPackage *> PackagesToCheck, const TMap<FSoftObjectPath, FSoftObjectPath>& AssetRedirectorMap) const
 {
 	// Add redirects as needed
diff --git a/Engine/Source/Developer/AssetTools/Private/AssetRenameManager.h b/Engine/Source/Developer/AssetTools/Private/AssetRenameManager.h
--- a/Engine/Source/Developer/AssetTools/Private/AssetRenameManager.h
+++ b/Engine/Source/Developer/AssetTools/Private/AssetRenameManager.h
@@ -65,9 +65,15 @@ private:
 	/** Attempts to load and fix redirector references for the supplied assets */
 	bool FixReferencesAndRename(const TArray<FAssetRenameData>& AssetsAndNames, bool bAutoCheckout, bool bWithDialog) const;
 
-	/** Get a list of assets referenced from CDOs */
-	TArray<TWeakObjectPtr<UObject>> FindCDOReferencedAssets(const TArray<FAssetRenameDataWithReferencers>& AssetsToRename) const;
-
+	/**
+	 * Get lists of assets with references from CDOs
+	 * @param AssetsToRename      Master list of assets to be renamed
+	 * @param OutReferences       Output array of hard asset references
+	 * @param OutSoftReferences   Output array of soft path asset references
+	 * @param bSetRedirectorFlags If true, each renamed asset with a soft reference will be marked as needing a redirector
+	 */
+	void FindCDOReferences(const TArrayView<FAssetRenameDataWithReferencers>& AssetsToRename, TArray<FAssetRenameDataWithReferencers*>& OutHardReferences, TArray<FAssetRenameDataWithReferencers*>& OutSoftReferences, bool bSetRedirectorFlags) const;
+	
 	/** Fills out the Referencing packages for all the assets described in AssetsToPopulate */
 	void PopulateAssetReferencers(TArray<FAssetRenameDataWithReferencers>& AssetsToPopulate) const;
 
diff --git a/Engine/Source/Editor/UnrealEd/Private/CookOnTheFlyServer.cpp b/Engine/Source/Editor/UnrealEd/Private/CookOnTheFlyServer.cpp
--- a/Engine/Source/Editor/UnrealEd/Private/CookOnTheFlyServer.cpp
+++ b/Engine/Source/Editor/UnrealEd/Private/CookOnTheFlyServer.cpp
@@ -5987,43 +5987,49 @@ void UCookOnTheFlyServer::CollectFilesToCook(TArray<FName>& FilesInPath, const T
 			FConfigFile PlatformEngineIni;
 			FConfigCacheIni::LoadLocalIniFile(PlatformEngineIni, TEXT("Engine"), true, *TargetPlatform->IniPlatformName());
 
-			// get the server and game default maps and cook them
-			FString Obj;
-			if (PlatformEngineIni.GetString(TEXT("/Script/EngineSettings.GameMapsSettings"), TEXT("GameDefaultMap"), Obj))
+			// get the server and game default maps/modes and cook them
+			TArray<FString, TInlineAllocator<5>> GameDefaultObjs; 
+			FString ObjPath;
+
+			if (PlatformEngineIni.GetString(TEXT("/Script/EngineSettings.GameMapsSettings"), TEXT("GameDefaultMap"), ObjPath))
 			{
-				if (Obj != FName(NAME_None).ToString())
-				{
-					AddFileToCook(FilesInPath, Obj);
-				}
+				GameDefaultObjs.Add(MoveTemp(ObjPath));
 			}
+
 			if (IsCookFlagSet(ECookInitializationFlags::IncludeServerMaps))
 			{
-				if (PlatformEngineIni.GetString(TEXT("/Script/EngineSettings.GameMapsSettings"), TEXT("ServerDefaultMap"), Obj))
+				if (PlatformEngineIni.GetString(TEXT("/Script/EngineSettings.GameMapsSettings"), TEXT("ServerDefaultMap"), ObjPath))
 				{
-					if (Obj != FName(NAME_None).ToString())
-					{
-						AddFileToCook(FilesInPath, Obj);
-					}
+					GameDefaultObjs.Add(MoveTemp(ObjPath));
 				}
 			}
-			if (PlatformEngineIni.GetString(TEXT("/Script/EngineSettings.GameMapsSettings"), TEXT("GlobalDefaultGameMode"), Obj))
+
+			if (PlatformEngineIni.GetString(TEXT("/Script/EngineSettings.GameMapsSettings"), TEXT("GlobalDefaultGameMode"), ObjPath))
 			{
-				if (Obj != FName(NAME_None).ToString())
-				{
-					AddFileToCook(FilesInPath, Obj);
-				}
+				GameDefaultObjs.Add(MoveTemp(ObjPath));
 			}
-			if (PlatformEngineIni.GetString(TEXT("/Script/EngineSettings.GameMapsSettings"), TEXT("GlobalDefaultServerGameMode"), Obj))
+
+			if (PlatformEngineIni.GetString(TEXT("/Script/EngineSettings.GameMapsSettings"), TEXT("GlobalDefaultServerGameMode"), ObjPath))
 			{
-				if (Obj != FName(NAME_None).ToString())
-				{
-					AddFileToCook(FilesInPath, Obj);
-				}
+				GameDefaultObjs.Add(MoveTemp(ObjPath));
+			}
+
+			if (PlatformEngineIni.GetString(TEXT("/Script/EngineSettings.GameMapsSettings"), TEXT("GameInstanceClass"), ObjPath))
+			{
+				GameDefaultObjs.Add(MoveTemp(ObjPath));
 			}
-			if (PlatformEngineIni.GetString(TEXT("/Script/EngineSettings.GameMapsSettings"), TEXT("GameInstanceClass"), Obj))
+
+			for (const FString& Obj : GameDefaultObjs)
 			{
 				if (Obj != FName(NAME_None).ToString())
 				{
+					FAssetData DestinationData = AssetRegistry->GetAssetByObjectPath(FName(*Obj), true);
+					if (DestinationData.IsRedirector())
+					{
+						const FText ErrorMessage = FText::Format(LOCTEXT("GameMapSettingsRedirectorDetected", "GameMapsSettings contains a redirected reference '{0}'. The intended asset will fail to load in a packaged build. Select the intended asset again in Project Settings to fix this issue."), FText::FromString(Obj));
+						LogCookerMessage(ErrorMessage.ToString(), EMessageSeverity::Error);
+					}
+
 					AddFileToCook(FilesInPath, Obj);
 				}
 			}
