Haptic Feedback types are now createable in the Content Browser Miscellaneous category again.

diff --git a/Engine/Source/Developer/AssetTools/Private/AssetTypeActions/AssetTypeActions_HapticFeedback.cpp b/Engine/Source/Developer/AssetTools/Private/AssetTypeActions/AssetTypeActions_HapticFeedback.cpp
new file mode 100644
--- /dev/null
+++ b/Engine/Source/Developer/AssetTools/Private/AssetTypeActions/AssetTypeActions_HapticFeedback.cpp
@@ -0,0 +1,21 @@
+// Copyright Epic Games, Inc. All Rights Reserved.
+
+#include "AssetTypeActions/AssetTypeActions_HapticFeedback.h"
+#include "Haptics/HapticFeedbackEffect_Buffer.h"
+#include "Haptics/HapticFeedbackEffect_Curve.h"
+#include "Haptics/HapticFeedbackEffect_SoundWave.h"
+
+UClass* FAssetTypeActions_HapticFeedbackEffectBuffer::GetSupportedClass() const
+{
+	return UHapticFeedbackEffect_Buffer::StaticClass();
+}
+
+UClass* FAssetTypeActions_HapticFeedbackEffectCurve::GetSupportedClass() const
+{
+	return UHapticFeedbackEffect_Curve::StaticClass();
+}
+
+UClass* FAssetTypeActions_HapticFeedbackEffectSoundWave::GetSupportedClass() const
+{
+	return UHapticFeedbackEffect_SoundWave::StaticClass();
+}
diff --git a/Engine/Source/Developer/AssetTools/Private/AssetTypeActions/AssetTypeActions_HapticFeedback.h b/Engine/Source/Developer/AssetTools/Private/AssetTypeActions/AssetTypeActions_HapticFeedback.h
new file mode 100644
--- /dev/null
+++ b/Engine/Source/Developer/AssetTools/Private/AssetTypeActions/AssetTypeActions_HapticFeedback.h
@@ -0,0 +1,36 @@
+// Copyright Epic Games, Inc. All Rights Reserved.
+
+#pragma once
+
+#include "CoreMinimal.h"
+#include "AssetTypeActions_Base.h"
+
+class FAssetTypeActions_HapticFeedbackEffectBuffer : public FAssetTypeActions_Base
+{
+public:
+	// IAssetTypeActions Implementation
+	virtual FText GetName() const override { return NSLOCTEXT("AssetTypeActions", "FAssetTypeActions_HapticFeedbackEffectBuffer", "Haptic Feedback Curve"); }
+	virtual FColor GetTypeColor() const override { return FColor(175, 0, 0); }
+	virtual UClass* GetSupportedClass() const override;
+	virtual uint32 GetCategories() override { return EAssetTypeCategories::Misc; }
+};
+
+class FAssetTypeActions_HapticFeedbackEffectCurve : public FAssetTypeActions_Base
+{
+public:
+	// IAssetTypeActions Implementation
+	virtual FText GetName() const override { return NSLOCTEXT("AssetTypeActions", "AssetTypeActions_HapticFeedbackEffectCurve", "Haptic Feedback Buffer"); }
+	virtual FColor GetTypeColor() const override { return FColor(175, 0, 0); }
+	virtual UClass* GetSupportedClass() const override;
+	virtual uint32 GetCategories() override { return EAssetTypeCategories::Misc; }
+};
+
+class FAssetTypeActions_HapticFeedbackEffectSoundWave : public FAssetTypeActions_Base
+{
+public:
+	// IAssetTypeActions Implementation
+	virtual FText GetName() const override { return NSLOCTEXT("AssetTypeActions", "FAssetTypeActions_HapticFeedbackEffectSoundWave", "Haptic Feedback Sound Wave"); }
+	virtual FColor GetTypeColor() const override { return FColor(175, 0, 0); }
+	virtual UClass* GetSupportedClass() const override;
+	virtual uint32 GetCategories() override { return EAssetTypeCategories::Misc; }
+};
\ No newline at end of file
