Haptic Feedback * assets are now createable in the content browser.

diff --git a/Engine/Source/Developer/AssetTools/Private/AssetTools.cpp b/Engine/Source/Developer/AssetTools/Private/AssetTools.cpp
--- a/Engine/Source/Developer/AssetTools/Private/AssetTools.cpp
+++ b/Engine/Source/Developer/AssetTools/Private/AssetTools.cpp
@@ -68,6 +68,7 @@
 #include "AssetTypeActions/AssetTypeActions_Font.h"
 #include "AssetTypeActions/AssetTypeActions_FontFace.h"
 #include "AssetTypeActions/AssetTypeActions_ForceFeedbackEffect.h"
+#include "AssetTypeActions/AssetTypeActions_HapticFeedback.h"
 #include "AssetTypeActions/AssetTypeActions_HLODProxy.h"
 #include "AssetTypeActions/AssetTypeActions_SubsurfaceProfile.h"
 #include "AssetTypeActions/AssetTypeActions_ActorFoliageSettings.h"
@@ -292,6 +293,9 @@ UAssetToolsImpl::UAssetToolsImpl(const FObjectInitializer& ObjectInitializer)
 	RegisterAssetTypeActions(MakeShareable(new FAssetTypeActions_VectorFieldAnimated));
 	RegisterAssetTypeActions(MakeShareable(new FAssetTypeActions_VectorFieldStatic));
 	RegisterAssetTypeActions(MakeShareable(new FAssetTypeActions_World));
+	RegisterAssetTypeActions(MakeShareable(new FAssetTypeActions_HapticFeedbackEffectBuffer));
+	RegisterAssetTypeActions(MakeShareable(new FAssetTypeActions_HapticFeedbackEffectCurve));
+	RegisterAssetTypeActions(MakeShareable(new FAssetTypeActions_HapticFeedbackEffectSoundWave));
 
 	// Note: Please don't add any more actions here!  They belong in an editor-only module that is more tightly
 	// coupled to your new system, and you should not create a dependency on your new system from AssetTools.
