Fix for bug where mirror data table does not appear in asset filter pop-out mirror

diff --git a/Engine/Source/Developer/AssetTools/Private/AssetTypeActions/AssetTypeActions_MirrorDataTable.h b/Engine/Source/Developer/AssetTools/Private/AssetTypeActions/AssetTypeActions_MirrorDataTable.h
--- a/Engine/Source/Developer/AssetTools/Private/AssetTypeActions/AssetTypeActions_MirrorDataTable.h
+++ b/Engine/Source/Developer/AssetTools/Private/AssetTypeActions/AssetTypeActions_MirrorDataTable.h
@@ -12,5 +12,6 @@ public:
 	virtual FText GetName() const override { return NSLOCTEXT("AssetTypeActions", "FAssetTypeActions_MirrorDataTable", "Mirror Data Table"); }
 	virtual UClass* GetSupportedClass() const override { return UMirrorDataTable::StaticClass(); }
 	virtual bool IsImportedAsset() const override { return false; }
+	virtual uint32 GetCategories() override { return EAssetTypeCategories::Animation; }
 	// End IAssetTypeActions
 };
