BPGC and WBPGC asset actions and icons

diff --git a/Engine/Source/Developer/AssetTools/Private/AssetTools.cpp b/Engine/Source/Developer/AssetTools/Private/AssetTools.cpp
--- a/Engine/Source/Developer/AssetTools/Private/AssetTools.cpp
+++ b/Engine/Source/Developer/AssetTools/Private/AssetTools.cpp
@@ -31,6 +31,7 @@
 #include "IClassTypeActions.h"
 #include "AssetTypeActions/AssetTypeActions_Actor.h"
 #include "AssetTypeActions/AssetTypeActions_Blueprint.h"
+#include "AssetTypeActions/AssetTypeActions_BlueprintGeneratedClass.h"
 #include "AssetTypeActions/AssetTypeActions_Curve.h"
 #include "AssetTypeActions/AssetTypeActions_MaterialInterface.h"
 #include "AssetTypeActions/AssetTypeActions_SkeletalMesh.h"
@@ -225,6 +226,7 @@ UAssetToolsImpl::UAssetToolsImpl(const FObjectInitializer& ObjectInitializer)
 	RegisterAssetTypeActions(MakeShareable(new FAssetTypeActions_PoseAsset));
 	RegisterAssetTypeActions(MakeShareable(new FAssetTypeActions_BlendSpace1D));
 	RegisterAssetTypeActions(MakeShareable(new FAssetTypeActions_Blueprint));
+	RegisterAssetTypeActions(MakeShareable(new FAssetTypeActions_BlueprintGeneratedClass));
 	RegisterAssetTypeActions(MakeShareable(new FAssetTypeActions_CameraAnim));
 	RegisterAssetTypeActions(MakeShareable(new FAssetTypeActions_CanvasRenderTarget2D));
 	RegisterAssetTypeActions(MakeShareable(new FAssetTypeActions_Curve));
diff --git a/Engine/Source/Developer/AssetTools/Private/AssetTypeActions/AssetTypeActions_BlueprintGeneratedClass.cpp b/Engine/Source/Developer/AssetTools/Private/AssetTypeActions/AssetTypeActions_BlueprintGeneratedClass.cpp
new file mode 100644
--- /dev/null
+++ b/Engine/Source/Developer/AssetTools/Private/AssetTypeActions/AssetTypeActions_BlueprintGeneratedClass.cpp
@@ -0,0 +1,135 @@
+ // Copyright Epic Games, Inc. All Rights Reserved.
+
+#include "AssetTypeActions/AssetTypeActions_BlueprintGeneratedClass.h"
+#include "Engine/BlueprintGeneratedClass.h"
+#include "ToolMenus.h"
+#include "Misc/PackageName.h"
+#include "Misc/MessageDialog.h"
+#include "Factories/BlueprintFactory.h"
+#include "ThumbnailRendering/SceneThumbnailInfo.h"
+#include "Kismet2/KismetEditorUtilities.h"
+#include "IContentBrowserSingleton.h"
+#include "ContentBrowserModule.h"
+
+#define LOCTEXT_NAMESPACE "AssetTypeActions"
+
+UClass* FAssetTypeActions_BlueprintGeneratedClass::GetSupportedClass() const
+{
+	return UBlueprintGeneratedClass::StaticClass();
+}
+
+void FAssetTypeActions_BlueprintGeneratedClass::GetActions(const TArray<UObject*>& InObjects, FToolMenuSection& Section)
+{
+	TArray<TWeakObjectPtr<UBlueprintGeneratedClass>> BPGCs = GetTypedWeakObjectPtrs<UBlueprintGeneratedClass>(InObjects);
+
+	if (BPGCs.Num() == 1)
+	{
+		TAttribute<FText>::FGetter DynamicTooltipGetter;
+		DynamicTooltipGetter.BindSP(this, &FAssetTypeActions_BlueprintGeneratedClass::GetNewDerivedBlueprintTooltip, BPGCs[0]);
+		TAttribute<FText> DynamicTooltipAttribute = TAttribute<FText>::Create(DynamicTooltipGetter);
+
+		Section.AddMenuEntry(
+			"BlueprintGeneratedClass_NewDerivedBlueprint",
+			LOCTEXT("BlueprintGeneratedClass_NewDerivedBlueprint", "Create Child Blueprint Class"),
+			DynamicTooltipAttribute,
+			FSlateIcon(FEditorStyle::GetStyleSetName(), "LevelEditor.CreateClassBlueprint"),
+			FUIAction(
+				FExecuteAction::CreateSP(this, &FAssetTypeActions_BlueprintGeneratedClass::ExecuteNewDerivedBlueprint, BPGCs[0]),
+				FCanExecuteAction::CreateSP(this, &FAssetTypeActions_BlueprintGeneratedClass::CanExecuteNewDerivedBlueprint, BPGCs[0])
+			)
+		);
+	}
+}
+
+UFactory* FAssetTypeActions_BlueprintGeneratedClass::GetFactoryForNewDerivedBlueprint(UBlueprintGeneratedClass* InBPGC) const
+{
+	UBlueprintFactory* Factory = NewObject<UBlueprintFactory>();
+	Factory->ParentClass = InBPGC;
+	return Factory;
+}
+
+UClass* FAssetTypeActions_BlueprintGeneratedClass::GetNewDerivedBlueprintClass() const
+{
+	return UBlueprint::StaticClass();
+}
+
+void FAssetTypeActions_BlueprintGeneratedClass::ExecuteNewDerivedBlueprint(TWeakObjectPtr<UBlueprintGeneratedClass> InObject)
+{
+	if (UBlueprintGeneratedClass* TargetParentClass = InObject.Get())
+	{
+		if (!FKismetEditorUtilities::CanCreateBlueprintOfClass(TargetParentClass))
+		{
+			FMessageDialog::Open(EAppMsgType::Ok, LOCTEXT("InvalidClassToMakeBlueprintFrom", "Invalid class with which to make a Blueprint."));
+			return;
+		}
+
+		FString Name;
+		FString PackageName;
+		CreateUniqueAssetName(TargetParentClass->GetOutermost()->GetName(), TEXT("_Child"), PackageName, Name);
+		const FString PackagePath = FPackageName::GetLongPackagePath(PackageName);
+
+		FContentBrowserModule& ContentBrowserModule = FModuleManager::LoadModuleChecked<FContentBrowserModule>("ContentBrowser");
+		ContentBrowserModule.Get().CreateNewAsset(Name, PackagePath, GetNewDerivedBlueprintClass(), GetFactoryForNewDerivedBlueprint(TargetParentClass));
+	}
+}
+
+FText FAssetTypeActions_BlueprintGeneratedClass::GetNewDerivedBlueprintTooltip(TWeakObjectPtr<UBlueprintGeneratedClass> InObject)
+{
+	if (!CanExecuteNewDerivedBlueprint(InObject))
+	{
+		return LOCTEXT("BlueprintGeneratedClass_NewDerivedBlueprintIsDeprecatedTooltip", "Compiled Blueprint class is deprecated, cannot derive a child Blueprint!");
+	}
+	else
+	{
+		return LOCTEXT("BlueprintGeneratedClass_NewDerivedBlueprintTooltip", "Creates a Child Blueprint Class based on the current compiled Blueprint, allowing you to create variants easily.");
+	}
+}
+
+bool FAssetTypeActions_BlueprintGeneratedClass::CanExecuteNewDerivedBlueprint(TWeakObjectPtr<UBlueprintGeneratedClass> InObject)
+{
+	UBlueprintGeneratedClass* BPGC = InObject.Get();
+	return BPGC && !BPGC->HasAnyClassFlags(CLASS_Deprecated);
+}
+
+//@todo
+
+//UThumbnailInfo* FAssetTypeActions_BlueprintGeneratedClass::GetThumbnailInfo(UObject* Asset) const
+//{
+//	// Blueprint thumbnail scenes are disabled for now
+//	UBlueprintGeneratedClass* BPGC = CastChecked<UBlueprintGeneratedClass>(Asset);
+//	UThumbnailInfo* ThumbnailInfo = BPGC->ThumbnailInfo;
+//	if (!ThumbnailInfo)
+//	{
+//		ThumbnailInfo = NewObject<USceneThumbnailInfo>(BPGC, NAME_None, RF_Transactional);
+//		BPGC->ThumbnailInfo = ThumbnailInfo;
+//	}
+//	return ThumbnailInfo;
+//}
+
+TWeakPtr<IClassTypeActions> FAssetTypeActions_BlueprintGeneratedClass::GetClassTypeActions(const FAssetData& AssetData) const
+{
+	// Blueprints get the class type actions for their parent native class.
+	// Using asset tags avoids us having to load the blueprint
+	UClass* ParentClass = nullptr;
+	FString ParentClassName;
+	if(!AssetData.GetTagValue(FBlueprintTags::NativeParentClassPath, ParentClassName))
+	{
+		AssetData.GetTagValue(FBlueprintTags::ParentClassPath, ParentClassName);
+	}
+	if (!ParentClassName.IsEmpty())
+	{
+		UObject* Outer = nullptr;
+		ResolveName(Outer, ParentClassName, false, false);
+		ParentClass = FindObject<UClass>(ANY_PACKAGE, *ParentClassName);
+	}
+
+	if (ParentClass)
+	{
+		FAssetToolsModule& AssetToolsModule = FAssetToolsModule::GetModule();
+		return AssetToolsModule.Get().GetClassTypeActionsForClass(ParentClass);
+	}
+
+	return nullptr;
+}
+
+#undef LOCTEXT_NAMESPACE
diff --git a/Engine/Source/Developer/AssetTools/Public/AssetTypeActions/AssetTypeActions_BlueprintGeneratedClass.h b/Engine/Source/Developer/AssetTools/Public/AssetTypeActions/AssetTypeActions_BlueprintGeneratedClass.h
new file mode 100644
--- /dev/null
+++ b/Engine/Source/Developer/AssetTools/Public/AssetTypeActions/AssetTypeActions_BlueprintGeneratedClass.h
@@ -0,0 +1,48 @@
+// Copyright Epic Games, Inc. All Rights Reserved.
+
+#pragma once
+
+#include "CoreMinimal.h"
+#include "Toolkits/IToolkitHost.h"
+#include "AssetTypeActions/AssetTypeActions_ClassTypeBase.h"
+
+class UBlueprintGeneratedClass;
+struct FAssetData;
+class IClassTypeActions;
+class UFactory;
+
+class ASSETTOOLS_API FAssetTypeActions_BlueprintGeneratedClass : public FAssetTypeActions_ClassTypeBase
+{
+public:
+	// IAssetTypeActions Implementation
+	virtual FText GetName() const override { return NSLOCTEXT("AssetTypeActions", "AssetTypeActions_BlueprintGeneratedClass", "Compiled Blueprint Class"); }
+	virtual FColor GetTypeColor() const override { return FColor(133, 173, 255); }
+	virtual UClass* GetSupportedClass() const override;
+	virtual bool HasActions(const TArray<UObject*>& InObjects) const override { return true; }
+	virtual void GetActions(const TArray<UObject*>& InObjects, FToolMenuSection& Section) override;
+	virtual uint32 GetCategories() override { return EAssetTypeCategories::Blueprint | EAssetTypeCategories::Basic; }
+	// @todo
+	//virtual class UThumbnailInfo* GetThumbnailInfo(UObject* Asset) const override;
+	// End IAssetTypeActions Implementation
+
+	// FAssetTypeActions_ClassTypeBase Implementation
+	virtual TWeakPtr<IClassTypeActions> GetClassTypeActions(const FAssetData& AssetData) const override;
+	// End FAssetTypeActions_ClassTypeBase Implementation
+
+protected:
+	/** Returns the factory responsible for creating a new Blueprint derived from the specified BPGC */
+	virtual UFactory* GetFactoryForNewDerivedBlueprint(UBlueprintGeneratedClass* InBPGC) const;
+
+	/** Returns the class of new Blueprints derived from this BPGC asset type */
+	virtual UClass* GetNewDerivedBlueprintClass() const;
+
+	/** Returns the tooltip to display when attempting to derive a Blueprint */
+	FText GetNewDerivedBlueprintTooltip(TWeakObjectPtr<UBlueprintGeneratedClass> InObject);
+
+	/** Returns TRUE if you can derive a Blueprint */
+	bool CanExecuteNewDerivedBlueprint(TWeakObjectPtr<UBlueprintGeneratedClass> InObject);
+
+private:
+	/** Handler for when NewDerivedBlueprint is selected */
+	void ExecuteNewDerivedBlueprint(TWeakObjectPtr<UBlueprintGeneratedClass> InObject);
+};
diff --git a/Engine/Source/Editor/EditorStyle/Private/SlateEditorStyle.cpp b/Engine/Source/Editor/EditorStyle/Private/SlateEditorStyle.cpp
--- a/Engine/Source/Editor/EditorStyle/Private/SlateEditorStyle.cpp
+++ b/Engine/Source/Editor/EditorStyle/Private/SlateEditorStyle.cpp
@@ -6441,6 +6441,7 @@ void FSlateEditorStyle::FStyle::SetupClassIconsAndThumbnails()
 			TEXT("BlockingVolume"),
 			TEXT("Blueprint"),
 			TEXT("BlueprintFunctionLibrary"),
+			TEXT("BlueprintGeneratedClass"),
 			TEXT("BlueprintInterface"),
 			TEXT("BlueprintMacroLibrary"),
 			TEXT("BoxReflectionCapture"),
@@ -6571,6 +6572,7 @@ void FSlateEditorStyle::FStyle::SetupClassIconsAndThumbnails()
 			TEXT("UserDefinedStruct"),
 			TEXT("UserWidget"),
 			TEXT("WidgetBlueprint"),
+			TEXT("WidgetBlueprintGeneratedClass"),
 			TEXT("WindDirectionalSource"),
 			TEXT("World"),
 			TEXT("Cube"),
diff --git a/Engine/Source/Editor/EditorStyle/Private/StarshipStyle.cpp b/Engine/Source/Editor/EditorStyle/Private/StarshipStyle.cpp
--- a/Engine/Source/Editor/EditorStyle/Private/StarshipStyle.cpp
+++ b/Engine/Source/Editor/EditorStyle/Private/StarshipStyle.cpp
@@ -5888,6 +5888,7 @@ void FStarshipEditorStyle::FStyle::SetupClassIconsAndThumbnails()
 			{TEXT("BlockingVolume")},
 			{TEXT("Blueprint")},
 			{TEXT("BlueprintFunctionLibrary")},
+			{TEXT("BlueprintGeneratedClass")},
 			{TEXT("BoxReflectionCapture")},
 			{TEXT("Brush")},
 			{TEXT("ButtonStyleAsset")},
@@ -6023,6 +6024,7 @@ void FStarshipEditorStyle::FStyle::SetupClassIconsAndThumbnails()
 			{TEXT("VolumetricCloud"), false},
 			{TEXT("VolumetricLightmapDensityVolume")},
 			{TEXT("WidgetBlueprint")},
+			{TEXT("WidgetBlueprintGeneratedClass")},
 			{TEXT("WindDirectionalSource")},
 			{TEXT("World")},
 		};
diff --git a/Engine/Source/Editor/UMGEditor/Private/AssetTypeActions_WidgetBlueprintGeneratedClass.cpp b/Engine/Source/Editor/UMGEditor/Private/AssetTypeActions_WidgetBlueprintGeneratedClass.cpp
new file mode 100644
--- /dev/null
+++ b/Engine/Source/Editor/UMGEditor/Private/AssetTypeActions_WidgetBlueprintGeneratedClass.cpp
@@ -0,0 +1,9 @@
+// Copyright Epic Games, Inc. All Rights Reserved.
+
+#include "AssetTypeActions_WidgetBlueprintGeneratedClass.h"
+#include "Blueprint/WidgetBlueprintGeneratedClass.h"
+
+UClass* FAssetTypeActions_WidgetBlueprintGeneratedClass::GetSupportedClass() const
+{
+	return UWidgetBlueprintGeneratedClass::StaticClass();
+}
diff --git a/Engine/Source/Editor/UMGEditor/Private/AssetTypeActions_WidgetBlueprintGeneratedClass.h b/Engine/Source/Editor/UMGEditor/Private/AssetTypeActions_WidgetBlueprintGeneratedClass.h
new file mode 100644
--- /dev/null
+++ b/Engine/Source/Editor/UMGEditor/Private/AssetTypeActions_WidgetBlueprintGeneratedClass.h
@@ -0,0 +1,17 @@
+// Copyright Epic Games, Inc. All Rights Reserved.
+
+#pragma once
+
+#include "CoreMinimal.h"
+#include "AssetTypeActions_Base.h"
+
+class FAssetTypeActions_WidgetBlueprintGeneratedClass : public FAssetTypeActions_Base
+{
+public:
+	// IAssetTypeActions Implementation
+	virtual FText GetName() const override { return NSLOCTEXT("AssetTypeActions", "AssetTypeActions_WidgetBlueprintGeneratedClass", "Compiled Widget Blueprint"); }
+	virtual FColor GetTypeColor() const override { return FColor(121,149,207); }
+	virtual UClass* GetSupportedClass() const override;
+	virtual uint32 GetCategories() override { return EAssetTypeCategories::UI; }
+	// End IAssetTypeActions Implementation
+};
diff --git a/Engine/Source/Editor/UMGEditor/Private/UMGEditorModule.cpp b/Engine/Source/Editor/UMGEditor/Private/UMGEditorModule.cpp
--- a/Engine/Source/Editor/UMGEditor/Private/UMGEditorModule.cpp
+++ b/Engine/Source/Editor/UMGEditor/Private/UMGEditorModule.cpp
@@ -14,6 +14,7 @@
 #include "AssetToolsModule.h"
 #include "IAssetTypeActions.h"
 #include "AssetTypeActions_WidgetBlueprint.h"
+#include "AssetTypeActions_WidgetBlueprintGeneratedClass.h"
 #include "KismetCompilerModule.h"
 #include "WidgetBlueprintCompiler.h"
 
@@ -73,6 +74,7 @@ public:
 		// Register asset types
 		IAssetTools& AssetTools = FModuleManager::LoadModuleChecked<FAssetToolsModule>("AssetTools").Get();
 		RegisterAssetTypeAction(AssetTools, MakeShareable(new FAssetTypeActions_WidgetBlueprint()));
+		RegisterAssetTypeAction(AssetTools, MakeShareable(new FAssetTypeActions_WidgetBlueprintGeneratedClass()));
 
 		FKismetCompilerContext::RegisterCompilerForBP(UWidgetBlueprint::StaticClass(), &UWidgetBlueprint::GetCompilerForWidgetBP );
 
diff --git a/Engine/Source/Editor/UMGEditor/Private/WidgetBlueprintCompiler.cpp b/Engine/Source/Editor/UMGEditor/Private/WidgetBlueprintCompiler.cpp
--- a/Engine/Source/Editor/UMGEditor/Private/WidgetBlueprintCompiler.cpp
+++ b/Engine/Source/Editor/UMGEditor/Private/WidgetBlueprintCompiler.cpp
@@ -703,15 +703,18 @@ void FWidgetBlueprintCompilerContext::FinishCompilingClass(UClass* Class)
 		// Make sure that we don't have dueling widget hierarchies
 		if (UWidgetBlueprintGeneratedClass* SuperBPGClass = Cast<UWidgetBlueprintGeneratedClass>(BPGClass->GetSuperClass()))
 		{
-			UWidgetBlueprint* SuperBlueprint = Cast<UWidgetBlueprint>(SuperBPGClass->ClassGeneratedBy);
-			if (SuperBlueprint->WidgetTree != nullptr)
+			if (SuperBPGClass->ClassGeneratedBy) // ClassGeneratedBy can be null for cooked widget blueprints
 			{
-				if ((SuperBlueprint->WidgetTree->RootWidget != nullptr) && (WidgetBlueprint()->WidgetTree->RootWidget != nullptr))
+				UWidgetBlueprint* SuperBlueprint = Cast<UWidgetBlueprint>(SuperBPGClass->ClassGeneratedBy);
+				if (ensure(SuperBlueprint) && SuperBlueprint->WidgetTree != nullptr)
 				{
-					// We both have a widget tree, terrible things will ensue
-					// @todo: nickd - we need to switch this back to a warning in engine, but note for games
-					MessageLog.Note(*LOCTEXT("ParentAndChildBothHaveWidgetTrees", "This widget @@ and parent class widget @@ both have a widget hierarchy, which is not supported.  Only one of them should have a widget tree.").ToString(),
-						WidgetBP, SuperBPGClass->ClassGeneratedBy);
+					if ((SuperBlueprint->WidgetTree->RootWidget != nullptr) && (WidgetBlueprint()->WidgetTree->RootWidget != nullptr))
+					{
+						// We both have a widget tree, terrible things will ensue
+						// @todo: nickd - we need to switch this back to a warning in engine, but note for games
+						MessageLog.Note(*LOCTEXT("ParentAndChildBothHaveWidgetTrees", "This widget @@ and parent class widget @@ both have a widget hierarchy, which is not supported.  Only one of them should have a widget tree.").ToString(),
+							WidgetBP, SuperBPGClass->ClassGeneratedBy);
+					}
 				}
 			}
 		}
diff --git a/Engine/Source/Editor/UnrealEd/Private/ClassIconFinder.cpp b/Engine/Source/Editor/UnrealEd/Private/ClassIconFinder.cpp
--- a/Engine/Source/Editor/UnrealEd/Private/ClassIconFinder.cpp
+++ b/Engine/Source/Editor/UnrealEd/Private/ClassIconFinder.cpp
@@ -113,7 +113,7 @@ const UClass* FClassIconFinder::GetIconClassForAssetData(const FAssetData& InAss
 	}
 	
 	static const FName IgnoreClassThumbnail(TEXT("IgnoreClassThumbnail"));
-	if ( AssetClass->IsChildOf<UBlueprint>() &&
+	if ( (AssetClass->IsChildOf<UBlueprint>() || AssetClass->IsChildOf<UBlueprintGeneratedClass>()) &&
 		!AssetClass->HasMetaDataHierarchical(IgnoreClassThumbnail))
 	{
 		if ( bOutIsClassType )
diff --git a/Engine/Source/Runtime/Engine/Classes/Engine/BlueprintGeneratedClass.h b/Engine/Source/Runtime/Engine/Classes/Engine/BlueprintGeneratedClass.h
--- a/Engine/Source/Runtime/Engine/Classes/Engine/BlueprintGeneratedClass.h
+++ b/Engine/Source/Runtime/Engine/Classes/Engine/BlueprintGeneratedClass.h
@@ -738,6 +738,7 @@ public:
 	virtual void PostLoad() override;
 	virtual void PostInitProperties() override;
 	virtual void GetPreloadDependencies(TArray<UObject*>& OutDeps) override;
+	virtual void GetAssetRegistryTags(TArray<FAssetRegistryTag>& OutTags) const override;
 	virtual FPrimaryAssetId GetPrimaryAssetId() const override;
 	virtual bool NeedsLoadForServer() const override;
 	virtual bool NeedsLoadForClient() const override;
diff --git a/Engine/Source/Runtime/Engine/Private/Blueprint.cpp b/Engine/Source/Runtime/Engine/Private/Blueprint.cpp
--- a/Engine/Source/Runtime/Engine/Private/Blueprint.cpp
+++ b/Engine/Source/Runtime/Engine/Private/Blueprint.cpp
@@ -901,7 +901,7 @@ void UBlueprint::GetAssetRegistryTags(TArray<FAssetRegistryTag>& OutTags) const
 	}
 	else
 	{
-		NativeParentClassName = ParentClassName = ("None");
+		NativeParentClassName = ParentClassName = TEXT("None");
 	}
 
 
diff --git a/Engine/Source/Runtime/Engine/Private/BlueprintGeneratedClass.cpp b/Engine/Source/Runtime/Engine/Private/BlueprintGeneratedClass.cpp
--- a/Engine/Source/Runtime/Engine/Private/BlueprintGeneratedClass.cpp
+++ b/Engine/Source/Runtime/Engine/Private/BlueprintGeneratedClass.cpp
@@ -159,6 +159,34 @@ void UBlueprintGeneratedClass::PostLoad()
 	AssembleReferenceTokenStream(true);
 }
 
+void UBlueprintGeneratedClass::GetAssetRegistryTags(TArray<FAssetRegistryTag>& OutTags) const
+{
+	Super::GetAssetRegistryTags(OutTags);
+
+	FString NativeParentClassName;
+	FString ParentClassName;
+
+	if (UClass* ParentClass = GetSuperClass())
+	{
+		ParentClassName = FString::Printf(TEXT("%s'%s'"), *ParentClass->GetClass()->GetName(), *ParentClass->GetPathName());
+
+		// Walk up until we find a native class (ie 'while they are BP classes')
+		UClass* NativeParentClass = ParentClass;
+		while (Cast<UBlueprintGeneratedClass>(NativeParentClass))
+		{
+			NativeParentClass = NativeParentClass->GetSuperClass();
+		}
+		NativeParentClassName = FString::Printf(TEXT("%s'%s'"), *NativeParentClass->GetClass()->GetName(), *NativeParentClass->GetPathName());
+	}
+	else
+	{
+		NativeParentClassName = ParentClassName = TEXT("None");
+	}
+
+	OutTags.Add(FAssetRegistryTag(FBlueprintTags::ParentClassPath, ParentClassName, FAssetRegistryTag::TT_Alphabetical));
+	OutTags.Add(FAssetRegistryTag(FBlueprintTags::NativeParentClassPath, NativeParentClassName, FAssetRegistryTag::TT_Alphabetical));
+}
+
 FPrimaryAssetId UBlueprintGeneratedClass::GetPrimaryAssetId() const
 {
 	FPrimaryAssetId AssetId;
