Migrate support: WorldPartition/OFPA Level

diff --git a/Engine/Source/Developer/AssetTools/Private/AssetTools.cpp b/Engine/Source/Developer/AssetTools/Private/AssetTools.cpp
--- a/Engine/Source/Developer/AssetTools/Private/AssetTools.cpp
+++ b/Engine/Source/Developer/AssetTools/Private/AssetTools.cpp
@@ -144,6 +144,8 @@
 #include "Misc/ConfigCacheIni.h"
 #include "Misc/BlacklistNames.h"
 #include "InterchangeManager.h"
+#include "Engine/World.h"
+#include "Engine/Level.h"
 
 #if WITH_EDITOR
 #include "Subsystems/AssetEditorSubsystem.h"
@@ -2724,6 +2726,7 @@ void UAssetToolsImpl::PerformMigratePackages(TArray<FName> PackageNamesToMigrate
 {
 	// Form a full list of packages to move by including the dependencies of the supplied packages
 	TSet<FName> AllPackageNamesToMove;
+	TSet<FString> ExternalActorsPaths;
 	{
 		FScopedSlowTask SlowTask( PackageNamesToMigrate.Num(), LOCTEXT( "MigratePackages_GatheringDependencies", "Gathering Dependencies..." ) );
 		SlowTask.MakeDialog();
@@ -2740,7 +2743,7 @@ void UAssetToolsImpl::PerformMigratePackages(TArray<FName> PackageNamesToMigrate
 				Path.RemoveFromStart(TEXT("/"));
 				Path.Split("/", &OriginalRootString, &Path, ESearchCase::IgnoreCase, ESearchDir::FromStart);
 				OriginalRootString = TEXT("/") + OriginalRootString;
-				RecursiveGetDependencies(*PackageIt, AllPackageNamesToMove, OriginalRootString);
+				RecursiveGetDependencies(*PackageIt, AllPackageNamesToMove, OriginalRootString, ExternalActorsPaths);
 			}
 		}
 	}
@@ -2997,7 +3000,7 @@ void UAssetToolsImpl::MigratePackages_ReportConfirmed(TSharedPtr<TArray<ReportPa
 	MigrateLog.Notify(LogMessage, Severity, true);
 }
 
-void UAssetToolsImpl::RecursiveGetDependencies(const FName& PackageName, TSet<FName>& AllDependencies, const FString& OriginalRoot) const
+void UAssetToolsImpl::RecursiveGetDependencies(const FName& PackageName, TSet<FName>& AllDependencies, const FString& OriginalRoot, TSet<FString>& ExternalActorsPaths) const
 {
 	FAssetRegistryModule& AssetRegistryModule = FModuleManager::Get().LoadModuleChecked<FAssetRegistryModule>(TEXT("AssetRegistry"));
 	TArray<FName> Dependencies;
@@ -3014,10 +3017,39 @@ void UAssetToolsImpl::RecursiveGetDependencies(const FName& PackageName, TSet<FN
 			if ( !bIsEnginePackage && !bIsScriptPackage && bIsInSamePackage )
 			{
 				AllDependencies.Add(*DependsIt);
-				RecursiveGetDependencies(*DependsIt, AllDependencies, OriginalRoot);
+				RecursiveGetDependencies(*DependsIt, AllDependencies, OriginalRoot, ExternalActorsPaths);
 			}
 		}
 	}
+
+	// Handle Specific External Actors use case (only used for the Migrate path for now)
+	// todo: revisit how to handle those in a more generic way.
+	TArray<FAssetData> Assets;
+	if (AssetRegistryModule.Get().GetAssetsByPackageName(PackageName, Assets))
+	{
+		for (const FAssetData& AssetData : Assets)
+		{
+			if (AssetData.GetClass() && AssetData.GetClass()->IsChildOf<UWorld>())
+			{
+				FString ExternalActorsPath = ULevel::GetExternalActorsPath(PackageName.ToString());
+				if (!ExternalActorsPath.IsEmpty() && !ExternalActorsPaths.Contains(ExternalActorsPath))
+				{
+					ExternalActorsPaths.Add(ExternalActorsPath);
+					AssetRegistryModule.Get().ScanPathsSynchronous({ ExternalActorsPath }, /*bForceRescan*/true, /*bIgnoreBlackListScanFilters*/true);
+
+					TArray<FAssetData> ExternalActorAssets;
+					AssetRegistryModule.Get().GetAssetsByPath(FName(*ExternalActorsPath), ExternalActorAssets, /*bRecursive*/true);
+
+					for (const FAssetData& ExternalActorAsset : ExternalActorAssets)
+					{
+						AllDependencies.Add(ExternalActorAsset.PackageName);
+						RecursiveGetDependencies(ExternalActorAsset.PackageName, AllDependencies, OriginalRoot, ExternalActorsPaths);
+					}
+				}
+			}
+		}
+	}
+	
 }
 
 void UAssetToolsImpl::RecursiveGetDependenciesAdvanced(const FName& PackageName, FAdvancedCopyParams& CopyParams, TArray<FName>& AllDependencies, TMap<FName, FName>& DependencyMap, const UAdvancedCopyCustomization* CopyCustomization, TArray<FAssetData>& OptionalAssetData) const
diff --git a/Engine/Source/Developer/AssetTools/Private/AssetTools.h b/Engine/Source/Developer/AssetTools/Private/AssetTools.h
--- a/Engine/Source/Developer/AssetTools/Private/AssetTools.h
+++ b/Engine/Source/Developer/AssetTools/Private/AssetTools.h
@@ -152,7 +152,7 @@ private:
 	void AdvancedCopyPackages_ReportConfirmed(FAdvancedCopyParams CopyParam, TArray<TMap<FString, FString>> DestinationMap) const;
 
 	/** Gets the dependencies of the specified package recursively */
-	void RecursiveGetDependencies(const FName& PackageName, TSet<FName>& AllDependencies, const FString& OriginalRoot) const;
+	void RecursiveGetDependencies(const FName& PackageName, TSet<FName>& AllDependencies, const FString& OriginalRoot, TSet<FString>& ExternalActorsPaths) const;
 
 	/** Gets the dependencies of the specified package recursively while omitting things that don't pass the FARFilter passed in from FAdvancedCopyParams */
 	void RecursiveGetDependenciesAdvanced(const FName& PackageName, FAdvancedCopyParams& CopyParams, TArray<FName>& AllDependencies, TMap<FName, FName>& DependencyMap, const class UAdvancedCopyCustomization* CopyCustomization, TArray<FAssetData>& OptionalAssetData) const;
