- Added button to Anim Modifiers details customizations to apply all instances of its class to their outer AnimSequences
- Added context menu entry to AnimationSequence asset actions to apply all contained animation modifiers

diff --git a/Engine/Source/Developer/AssetTools/Private/AssetTypeActions/AssetTypeActions_AnimSequence.cpp b/Engine/Source/Developer/AssetTools/Private/AssetTypeActions/AssetTypeActions_AnimSequence.cpp
--- a/Engine/Source/Developer/AssetTools/Private/AssetTypeActions/AssetTypeActions_AnimSequence.cpp
+++ b/Engine/Source/Developer/AssetTools/Private/AssetTypeActions/AssetTypeActions_AnimSequence.cpp
@@ -49,12 +49,63 @@ void FAssetTypeActions_AnimSequence::GetActions(const TArray<UObject*>& InObject
 		FUIAction(FExecuteAction::CreateSP(this, &FAssetTypeActions_AnimSequence::ExecuteReimportWithNewSource, Sequences))
 		);
 
-	Section.AddMenuEntry(
-		"AnimSequence_AddAnimationModifier",
-		LOCTEXT("AnimSequence_AddAnimationModifier", "Add Animation Modifier(s)"),
-		LOCTEXT("AnimSequence_AddAnimationModifierTooltip", "Apply new animation modifier(s)."),
-		FSlateIcon(FEditorStyle::GetStyleSetName(), "ClassIcon.AnimationModifier"),
-		FUIAction(FExecuteAction::CreateSP(this, &FAssetTypeActions_AnimSequence::ExecuteAddNewAnimationModifier, Sequences))
+	FNewMenuDelegate MenuDelegate = FNewMenuDelegate::CreateLambda([Sequences](FMenuBuilder& MenuBuilder)
+	{
+		MenuBuilder.AddMenuEntry(
+	        LOCTEXT("AnimSequence_AddAnimationModifier", "Add Modifiers"),
+	        LOCTEXT("AnimSequence_AddAnimationModifierTooltip", "Add new animation modifier(s)."),
+	       FSlateIcon(FEditorStyle::GetStyleSetName(), "ClassIcon.AnimationModifier"),
+	        FUIAction(FExecuteAction::CreateLambda([Sequences]()
+	        {
+		        TArray<UAnimSequence*> AnimSequences;
+
+			    Algo::TransformIf(Sequences, AnimSequences, 
+			    [](const TWeakObjectPtr<UAnimSequence>& WeakAnimSequence)
+			    {
+			        return WeakAnimSequence.Get() && WeakAnimSequence->IsA<UAnimSequence>();
+			    },
+			    [](const TWeakObjectPtr<UAnimSequence>& WeakAnimSequence)
+			    {
+			        return WeakAnimSequence.Get();
+			    });
+
+			    if (IAnimationModifiersModule* Module = FModuleManager::Get().LoadModulePtr<IAnimationModifiersModule>("AnimationModifiers"))
+			    {
+			        Module->ShowAddAnimationModifierWindow(AnimSequences);
+			    }
+	        }))
+	    );
+	
+	    MenuBuilder.AddMenuEntry(
+	        LOCTEXT("AnimSequence_ApplyAnimationModifier", "Apply Modifiers"),
+	        LOCTEXT("AnimSequence_ApplyAnimationModifierTooltip", "Applies all contained animation modifier(s)."),
+			FSlateIcon(FEditorStyle::GetStyleSetName(), "ClassIcon.AnimationModifier"),
+	        FUIAction(FExecuteAction::CreateLambda([Sequences]()
+	        {
+		        TArray<UAnimSequence*> AnimSequences;
+				Algo::TransformIf(Sequences, AnimSequences, 
+				[](const TWeakObjectPtr<UAnimSequence>& WeakAnimSequence)
+				{
+				    return WeakAnimSequence.Get() && WeakAnimSequence->IsA<UAnimSequence>();
+				},
+				[](const TWeakObjectPtr<UAnimSequence>& WeakAnimSequence)
+				{
+				    return WeakAnimSequence.Get();
+				});
+
+				if (IAnimationModifiersModule* Module = FModuleManager::Get().LoadModulePtr<IAnimationModifiersModule>("AnimationModifiers"))
+				{
+				    Module->ApplyAnimationModifiers(AnimSequences);
+				}
+	        }))
+	    );
+	});
+	
+	Section.AddSubMenu("AnimSequence_AnimationModifiers", LOCTEXT("AnimSequence_AnimationModifiers", "Animation Modifier(s)"),
+	LOCTEXT("AnimSequence_AnimationModifiersTooltip", "Animation Modifier actions"),
+		FNewToolMenuChoice(MenuDelegate),
+		false,
+		FSlateIcon(FEditorStyle::GetStyleSetName(), "ClassIcon.AnimationModifier")
 	);
 
 	FAssetTypeActions_AnimationAsset::GetActions(InObjects, Section);
@@ -172,26 +223,6 @@ void FAssetTypeActions_AnimSequence::ExecuteNewPoseAsset(TArray<TWeakObjectPtr<U
 	CreateAnimationAssets(Objects, UPoseAsset::StaticClass(), Factory, DefaultSuffix, FOnConfigureFactory::CreateSP(this, &FAssetTypeActions_AnimSequence::ConfigureFactoryForPoseAsset));
 }
 
-void FAssetTypeActions_AnimSequence::ExecuteAddNewAnimationModifier(TArray<TWeakObjectPtr<UAnimSequence>> Objects)
-{
-	TArray<UAnimSequence*> AnimSequences;
-
-	Algo::TransformIf(Objects, AnimSequences, 
-	[](const TWeakObjectPtr<UAnimSequence>& WeakAnimSequence)
-	{
-		return WeakAnimSequence.Get() && WeakAnimSequence->IsA<UAnimSequence>();
-	},
-	[](const TWeakObjectPtr<UAnimSequence>& WeakAnimSequence)
-	{
-		return WeakAnimSequence.Get();
-	});
-
-	if (IAnimationModifiersModule* Module = FModuleManager::Get().LoadModulePtr<IAnimationModifiersModule>("AnimationModifiers"))
-	{
-		Module->ShowAddAnimationModifierWindow(AnimSequences);
-	}
-}
-
 bool FAssetTypeActions_AnimSequence::ConfigureFactoryForAnimComposite(UFactory* AssetFactory, UAnimSequence* SourceAnimation) const
 {
 	UAnimCompositeFactory* CompositeFactory = CastChecked<UAnimCompositeFactory>(AssetFactory);
diff --git a/Engine/Source/Developer/AssetTools/Private/AssetTypeActions/AssetTypeActions_AnimSequence.h b/Engine/Source/Developer/AssetTools/Private/AssetTypeActions/AssetTypeActions_AnimSequence.h
--- a/Engine/Source/Developer/AssetTools/Private/AssetTypeActions/AssetTypeActions_AnimSequence.h
+++ b/Engine/Source/Developer/AssetTools/Private/AssetTypeActions/AssetTypeActions_AnimSequence.h
@@ -40,9 +40,6 @@ private:
 	/** Handler for when Create Pose Asset is selected */
 	void ExecuteNewPoseAsset(TArray<TWeakObjectPtr<UAnimSequence>> Objects) const;
 
-	/** Handler for when Add Animation Modifier is selected */
-	void ExecuteAddNewAnimationModifier(TArray<TWeakObjectPtr<UAnimSequence>> Objects);
-
 	/** Delegate handler passed to CreateAnimationAssets when creating AnimComposites */
 	bool ConfigureFactoryForAnimComposite(UFactory* AssetFactory, UAnimSequence* SourceAnimation) const;
 
diff --git a/Engine/Source/Editor/AnimationModifiers/Private/AnimationModifier.cpp b/Engine/Source/Editor/AnimationModifiers/Private/AnimationModifier.cpp
--- a/Engine/Source/Editor/AnimationModifiers/Private/AnimationModifier.cpp
+++ b/Engine/Source/Editor/AnimationModifiers/Private/AnimationModifier.cpp
@@ -1,9 +1,11 @@
 // Copyright Epic Games, Inc. All Rights Reserved.
 
 #include "AnimationModifier.h"
+#include "AssetViewUtils.h"
 #include "Animation/AnimSequence.h"
 #include "Animation/Skeleton.h"
 #include "ModifierOutputFilter.h"
+#include "AssetRegistry/AssetRegistryModule.h"
 #include "Editor/Transactor.h"
 #include "UObject/UObjectIterator.h"
 
@@ -170,13 +172,6 @@ void UAnimationModifier::PostInitProperties()
 {
 	Super::PostInitProperties();
 	UpdateNativeRevisionGuid();
-
-	// Ensure we always have a valid guid
-	if (!RevisionGuid.IsValid())
-	{
-		UpdateRevisionGuid(GetClass());
-		MarkPackageDirty();
-	}
 }
 
 void UAnimationModifier::Serialize(FArchive& Ar)
@@ -191,6 +186,17 @@ void UAnimationModifier::Serialize(FArchive& Ar)
 	}
 }
 
+void UAnimationModifier::PostLoad()
+{
+	Super::PostLoad();	
+	// Ensure we always have a valid guid
+	if (!RevisionGuid.IsValid())
+	{
+		UpdateRevisionGuid(GetClass());
+		MarkPackageDirty();
+	}
+}
+
 const USkeleton* UAnimationModifier::GetSkeleton()
 {
 	return CurrentSkeleton;
@@ -198,14 +204,19 @@ const USkeleton* UAnimationModifier::GetSkeleton()
 
 void UAnimationModifier::UpdateRevisionGuid(UClass* ModifierClass)
 {
-	RevisionGuid = FGuid::NewGuid();
-
-	// Native classes are more difficult?
-	for (TObjectIterator<UAnimationModifier> It; It; ++It)
+	if (ModifierClass)
 	{
-		if (*It != this && It->GetClass() == ModifierClass)
+		LoadModifierReferencers(ModifierClass);
+	
+		RevisionGuid = FGuid::NewGuid();
+
+		// Native classes are more difficult?
+		for (TObjectIterator<UAnimationModifier> It; It; ++It)
 		{
-			It->SetInstanceRevisionGuid(RevisionGuid);
+			if (*It != this && It->GetClass() == ModifierClass)
+			{
+				It->SetInstanceRevisionGuid(RevisionGuid);
+			}
 		}
 	}
 }
@@ -232,6 +243,49 @@ void UAnimationModifier::UpdateNativeRevisionGuid()
 	}
 }
 
+void UAnimationModifier::ApplyToAll(TSubclassOf<UAnimationModifier> ModifierSubClass)
+{
+	if (UClass* ModifierClass = ModifierSubClass.Get())
+	{
+		// Make sure all packages (in this case UAnimSequences) are loaded to ensure the TObjectIterator has any instances to iterate over
+		LoadModifierReferencers(ModifierSubClass);
+		
+		const FScopedTransaction Transaction(LOCTEXT("UndoAction_ApplyModifiers", "Applying Animation Modifier to Animation Sequence(s)"));		
+		for (TObjectIterator<UAnimationModifier> It; It; ++It)
+		{
+			if (*It && It->GetClass() == ModifierClass)
+			{
+				// Go through outer chain to find AnimSequence
+				UObject* Outer = It->GetOuter();
+				while(Outer && !Outer->IsA<UAnimSequence>())
+				{
+					Outer = Outer->GetOuter();
+				}
+
+				if (UAnimSequence* AnimSequence = Cast<UAnimSequence>(Outer))
+				{
+					AnimSequence->Modify();
+					It->ApplyToAnimationSequence(AnimSequence);
+				}			
+			}
+		}
+	}	
+}
+
+void UAnimationModifier::LoadModifierReferencers(TSubclassOf<UAnimationModifier> ModifierSubClass)
+{
+	if (UClass* ModifierClass = ModifierSubClass.Get())
+	{
+		FAssetRegistryModule& AssetRegistryModule = FModuleManager::LoadModuleChecked<FAssetRegistryModule>(TEXT("AssetRegistry"));
+		TArray<FName> PackageDependencies;
+		AssetRegistryModule.GetRegistry().GetReferencers(ModifierClass->GetPackage()->GetFName(), PackageDependencies);
+
+		TArray<FString> PackageNames;
+		Algo::Transform(PackageDependencies, PackageNames, [](FName Name) { return Name.ToString(); });
+		TArray<UPackage*> Packages = AssetViewUtils::LoadPackages(PackageNames);
+	}
+}
+
 int32 UAnimationModifier::GetNativeClassRevision() const
 {
 	// Overriden in derrived classes to perform native revisioning
@@ -258,4 +312,4 @@ void UAnimationModifier::SetInstanceRevisionGuid(FGuid Guid)
 	RevisionGuid = Guid;
 }
 
-#undef LOCTEXT_NAMESPACE
\ No newline at end of file
+#undef LOCTEXT_NAMESPACE
diff --git a/Engine/Source/Editor/AnimationModifiers/Private/AnimationModifierDetailCustomization.cpp b/Engine/Source/Editor/AnimationModifiers/Private/AnimationModifierDetailCustomization.cpp
--- a/Engine/Source/Editor/AnimationModifiers/Private/AnimationModifierDetailCustomization.cpp
+++ b/Engine/Source/Editor/AnimationModifiers/Private/AnimationModifierDetailCustomization.cpp
@@ -29,19 +29,29 @@ void FAnimationModifierDetailCustomization::CustomizeDetails(IDetailLayoutBuilde
 	// If we have found a valid modifier instance add a revision bump button to the details panel
 	if (ModifierInstance)
 	{
-		IDetailCategoryBuilder& RevisionCategory = DetailBuilder.EditCategory("Revision");
-		FDetailWidgetRow& UpdateRevisionRow = RevisionCategory.AddCustomRow(LOCTEXT("UpdateRevisionSearchLabel", "Update Revision"))
+		IDetailCategoryBuilder& RevisionCategory = DetailBuilder.EditCategory("Modifier Instances");
+		FDetailWidgetRow& UpdateRevisionRow = RevisionCategory.AddCustomRow(LOCTEXT("ModifierActionsLabel", "Modifier Actions"))
 		.WholeRowWidget
 		[
-			SAssignNew(UpdateRevisionButton, SButton)
-			.OnClicked(this, &FAnimationModifierDetailCustomization::OnUpdateRevisionButtonClicked)
-			.ContentPadding(FMargin(2))
-			.Content()
+			SNew(SHorizontalBox)
+			+SHorizontalBox::Slot()
+           .AutoWidth()
+           .Padding(2.f)
+           [
+               SNew(SButton)
+               .OnClicked(this, &FAnimationModifierDetailCustomization::OnApplyButtonClicked)
+               .ToolTipText(LOCTEXT("ApplyToolTip", "Applies any instanced modifiers of this class to their owning Animation Sequences."))
+               .Text(LOCTEXT("ApplyText", "Apply to All"))
+           ]
+			+SHorizontalBox::Slot()
+			.AutoWidth()
+			.Padding(2.f)
 			[
-				SNew(STextBlock)
-				.Justification(ETextJustify::Center)
+				SAssignNew(UpdateRevisionButton, SButton)
+				.OnClicked(this, &FAnimationModifierDetailCustomization::OnUpdateRevisionButtonClicked)
+				.ToolTipText(LOCTEXT("UpdateRevisionToolTip", "Updates the stored revision GUID on all instances of this Modifier class, marking them out-of-date."))
 				.Text(LOCTEXT("UpdateRevisionText", "Update Revision"))
-			]
+            ]           
 		];
 	}
 }
@@ -55,4 +65,13 @@ FReply FAnimationModifierDetailCustomization::OnUpdateRevisionButtonClicked()
 	return FReply::Handled();
 }
 
+FReply FAnimationModifierDetailCustomization::OnApplyButtonClicked()
+{
+	if (ModifierInstance)
+	{
+		UAnimationModifier::ApplyToAll(ModifierInstance->GetClass());
+	}
+	return FReply::Handled();
+}
+
 #undef LOCTEXT_NAMESPACE //"FAnimationModifierDetailCustomization"
diff --git a/Engine/Source/Editor/AnimationModifiers/Private/AnimationModifierDetailCustomization.h b/Engine/Source/Editor/AnimationModifiers/Private/AnimationModifierDetailCustomization.h
--- a/Engine/Source/Editor/AnimationModifiers/Private/AnimationModifierDetailCustomization.h
+++ b/Engine/Source/Editor/AnimationModifiers/Private/AnimationModifierDetailCustomization.h
@@ -16,6 +16,7 @@ public:
 	virtual void CustomizeDetails(IDetailLayoutBuilder& DetailBuilder) override;
 protected:
 	FReply OnUpdateRevisionButtonClicked();
+	FReply OnApplyButtonClicked();
 protected:
 	TSharedPtr<SButton> UpdateRevisionButton;
 	UAnimationModifier* ModifierInstance;
diff --git a/Engine/Source/Editor/AnimationModifiers/Private/AnimationModifiersModule.cpp b/Engine/Source/Editor/AnimationModifiers/Private/AnimationModifiersModule.cpp
--- a/Engine/Source/Editor/AnimationModifiers/Private/AnimationModifiersModule.cpp
+++ b/Engine/Source/Editor/AnimationModifiers/Private/AnimationModifiersModule.cpp
@@ -16,6 +16,8 @@
 #include "Framework/Application/SlateApplication.h"
 #include "Interfaces/IMainFrameModule.h"
 
+#include "AnimationModifiersAssetUserData.h"
+
 #define LOCTEXT_NAMESPACE "AnimationModifiersModule"
 
 IMPLEMENT_MODULE(FAnimationModifiersModule, AnimationModifiers);
@@ -95,4 +97,28 @@ void FAnimationModifiersModule::ShowAddAnimationModifierWindow(const TArray<UAni
 	FSlateApplication::Get().AddModalWindow(Window, ParentWindow, false);
 }
 
+void FAnimationModifiersModule::ApplyAnimationModifiers(const TArray<UAnimSequence*>& InSequences)
+{
+	const FScopedTransaction Transaction(LOCTEXT("UndoAction_ApplyModifiers", "Applying Animation Modifier(s) to Animation Sequence(s)"));
+	
+	// Iterate over each Animation Sequence and all of its contained modifiers, applying each one
+	TArray<UAnimationModifiersAssetUserData*> AssetUserData;
+	for (UAnimSequence* AnimationSequence : InSequences)
+	{
+		if (AnimationSequence)
+		{
+			UAnimationModifiersAssetUserData* UserData = AnimationSequence->GetAssetUserData<UAnimationModifiersAssetUserData>();
+			if (UserData)
+			{
+				AnimationSequence->Modify();
+				const TArray<UAnimationModifier*>& ModifierInstances = UserData->GetAnimationModifierInstances();
+				for (UAnimationModifier* Modifier : ModifierInstances)
+				{
+					Modifier->ApplyToAnimationSequence(AnimationSequence);
+				}
+			}
+		}		
+	}
+}
+
 #undef LOCTEXT_NAMESPACE // "AnimationModifiersModule"
diff --git a/Engine/Source/Editor/AnimationModifiers/Private/AnimationModifiersModule.h b/Engine/Source/Editor/AnimationModifiers/Private/AnimationModifiersModule.h
--- a/Engine/Source/Editor/AnimationModifiers/Private/AnimationModifiersModule.h
+++ b/Engine/Source/Editor/AnimationModifiers/Private/AnimationModifiersModule.h
@@ -20,6 +20,7 @@ public:
 
 	/** Begin IAnimationModifiersModule overrides */
 	virtual void ShowAddAnimationModifierWindow(const TArray<UAnimSequence*>& InSequences) override;
+	virtual void ApplyAnimationModifiers(const TArray<UAnimSequence*>& InSequences) override;
 	/** End IAnimationModifiersModule overrides */
 
 protected:
diff --git a/Engine/Source/Editor/AnimationModifiers/Public/AnimationModifier.h b/Engine/Source/Editor/AnimationModifiers/Public/AnimationModifier.h
--- a/Engine/Source/Editor/AnimationModifiers/Public/AnimationModifier.h
+++ b/Engine/Source/Editor/AnimationModifiers/Public/AnimationModifier.h
@@ -41,6 +41,7 @@ public:
 	// Begin UObject Overrides
 	virtual void PostInitProperties() override;
 	virtual void Serialize(FArchive& Ar) override;
+	virtual void PostLoad() override;
 	// End UObject Overrides
 protected:
 	// Derived class accessors to skeleton and anim sequence 
@@ -56,6 +57,10 @@ protected:
 	/** Updating of blueprint and native GUIDs*/
 	void UpdateRevisionGuid(UClass* ModifierClass);
 	void UpdateNativeRevisionGuid();
+
+	/** Applies all instances of the provided Modifier class to its outer Animation Sequence*/
+	static void ApplyToAll(TSubclassOf<UAnimationModifier> ModifierSubClass);
+	static void LoadModifierReferencers(TSubclassOf<UAnimationModifier> ModifierSubClass);
 private:
 	UAnimSequence* CurrentAnimSequence;
 	USkeleton* CurrentSkeleton;
diff --git a/Engine/Source/Editor/AnimationModifiers/Public/IAnimationModifiersModule.h b/Engine/Source/Editor/AnimationModifiers/Public/IAnimationModifiersModule.h
--- a/Engine/Source/Editor/AnimationModifiers/Public/IAnimationModifiersModule.h
+++ b/Engine/Source/Editor/AnimationModifiers/Public/IAnimationModifiersModule.h
@@ -12,4 +12,7 @@ class IAnimationModifiersModule : public IModuleInterface
 public:
 	/** Shows a new modal dialog allowing the user to setup Animation Modifiers to be added for all AnimSequences part of InSequences */
 	virtual void ShowAddAnimationModifierWindow(const TArray<UAnimSequence*>& InSequences) = 0;
+
+	/** Apply animation modifiers contained by provided Animation Sequences to its respective outers */
+	virtual void ApplyAnimationModifiers(const TArray<UAnimSequence*>& InSequences) = 0;	
 };
