Fixed transposed name and type in the Haptic Feedback Effect AssetTypeActions implementations.

diff --git a/Engine/Source/Developer/AssetTools/Private/AssetTypeActions/AssetTypeActions_HapticFeedback.h b/Engine/Source/Developer/AssetTools/Private/AssetTypeActions/AssetTypeActions_HapticFeedback.h
--- a/Engine/Source/Developer/AssetTools/Private/AssetTypeActions/AssetTypeActions_HapticFeedback.h
+++ b/Engine/Source/Developer/AssetTools/Private/AssetTypeActions/AssetTypeActions_HapticFeedback.h
@@ -9,7 +9,7 @@ class FAssetTypeActions_HapticFeedbackEffectBuffer : public FAssetTypeActions_Ba
 {
 public:
 	// IAssetTypeActions Implementation
-	virtual FText GetName() const override { return NSLOCTEXT("AssetTypeActions", "FAssetTypeActions_HapticFeedbackEffectBuffer", "Haptic Feedback Curve"); }
+	virtual FText GetName() const override { return NSLOCTEXT("AssetTypeActions", "FAssetTypeActions_HapticFeedbackEffectBuffer", "Haptic Feedback Buffer"); }
 	virtual FColor GetTypeColor() const override { return FColor(175, 0, 0); }
 	virtual UClass* GetSupportedClass() const override;
 	virtual uint32 GetCategories() override { return EAssetTypeCategories::Misc; }
@@ -19,7 +19,7 @@ class FAssetTypeActions_HapticFeedbackEffectCurve : public FAssetTypeActions_Bas
 {
 public:
 	// IAssetTypeActions Implementation
-	virtual FText GetName() const override { return NSLOCTEXT("AssetTypeActions", "AssetTypeActions_HapticFeedbackEffectCurve", "Haptic Feedback Buffer"); }
+	virtual FText GetName() const override { return NSLOCTEXT("AssetTypeActions", "AssetTypeActions_HapticFeedbackEffectCurve", "Haptic Feedback Curve"); }
 	virtual FColor GetTypeColor() const override { return FColor(175, 0, 0); }
 	virtual UClass* GetSupportedClass() const override;
 	virtual uint32 GetCategories() override { return EAssetTypeCategories::Misc; }
