Avoid stall when refreshing content browser thumbnails of imported skeletal mesh by avoiding querying the overlay type as it will be refresh when compilation finishes

diff --git a/Engine/Source/Developer/AssetTools/Private/AssetTypeActions/AssetTypeActions_SkeletalMesh.cpp b/Engine/Source/Developer/AssetTools/Private/AssetTypeActions/AssetTypeActions_SkeletalMesh.cpp
--- a/Engine/Source/Developer/AssetTools/Private/AssetTypeActions/AssetTypeActions_SkeletalMesh.cpp
+++ b/Engine/Source/Developer/AssetTools/Private/AssetTypeActions/AssetTypeActions_SkeletalMesh.cpp
@@ -603,13 +603,17 @@ EVisibility FAssetTypeActions_SkeletalMesh::GetThumbnailSkinningOverlayVisibilit
 	UObject* Obj = AssetData.GetAsset();
 	if(USkeletalMesh* SkeletalMesh = Cast<USkeletalMesh>(Obj))
 	{
-		UAssetImportData* GenericImportData = SkeletalMesh->GetAssetImportData();
-		if (GenericImportData != nullptr)
+		// Don't block on asset compilation just for overlay, the thumbnail is refreshed anyway once the compilation is complete.
+		if (!SkeletalMesh->IsCompiling())
 		{
-			UFbxSkeletalMeshImportData* ImportData = Cast<UFbxSkeletalMeshImportData>(GenericImportData);
-			if (ImportData != nullptr && ImportData->LastImportContentType == EFBXImportContentType::FBXICT_Geometry)
+			UAssetImportData* GenericImportData = SkeletalMesh->GetAssetImportData();
+			if (GenericImportData != nullptr)
 			{
-				return EVisibility::HitTestInvisible;
+				UFbxSkeletalMeshImportData* ImportData = Cast<UFbxSkeletalMeshImportData>(GenericImportData);
+				if (ImportData != nullptr && ImportData->LastImportContentType == EFBXImportContentType::FBXICT_Geometry)
+				{
+					return EVisibility::HitTestInvisible;
+				}
 			}
 		}
 	}
