Added AssetTypeActions for Sub UV Animation, which allows the asset type to be created again.

diff --git a/Engine/Source/Developer/AssetTools/Private/AssetTools.cpp b/Engine/Source/Developer/AssetTools/Private/AssetTools.cpp
--- a/Engine/Source/Developer/AssetTools/Private/AssetTools.cpp
+++ b/Engine/Source/Developer/AssetTools/Private/AssetTools.cpp
@@ -99,6 +99,7 @@
 #include "AssetTypeActions/AssetTypeActions_SlateBrush.h"
 #include "AssetTypeActions/AssetTypeActions_SlateWidgetStyle.h"
 #include "AssetTypeActions/AssetTypeActions_StaticMesh.h"
+#include "AssetTypeActions/AssetTypeActions_SubUVAnimation.h"
 #include "AssetTypeActions/AssetTypeActions_Texture2D.h"
 #include "AssetTypeActions/AssetTypeActions_Texture2DArray.h"
 #include "AssetTypeActions/AssetTypeActions_TextureCube.h"
@@ -271,6 +272,7 @@ UAssetToolsImpl::UAssetToolsImpl(const FObjectInitializer& ObjectInitializer)
 	RegisterAssetTypeActions(MakeShareable(new FAssetTypeActions_MirrorDataTable));
 	RegisterAssetTypeActions(MakeShareable(new FAssetTypeActions_ObjectLibrary));
 	RegisterAssetTypeActions(MakeShareable(new FAssetTypeActions_ParticleSystem));
+	RegisterAssetTypeActions(MakeShareable(new FAssetTypeActions_SubUVAnimation));
 	RegisterAssetTypeActions(MakeShareable(new FAssetTypeActions_PhysicalMaterial));
 	RegisterAssetTypeActions(MakeShareable(new FAssetTypeActions_PhysicalMaterialMask));
 	RegisterAssetTypeActions(MakeShareable(new FAssetTypeActions_PhysicsAsset));
diff --git a/Engine/Source/Developer/AssetTools/Private/AssetTypeActions/AssetTypeActions_SubUVAnimation.h b/Engine/Source/Developer/AssetTools/Private/AssetTypeActions/AssetTypeActions_SubUVAnimation.h
new file mode 100644
--- /dev/null
+++ b/Engine/Source/Developer/AssetTools/Private/AssetTypeActions/AssetTypeActions_SubUVAnimation.h
@@ -0,0 +1,18 @@
+// Copyright Epic Games, Inc. All Rights Reserved.
+
+#pragma once
+
+#include "CoreMinimal.h"
+#include "AssetTypeActions_Base.h"
+#include "Particles/SubUVAnimation.h"
+
+class FAssetTypeActions_SubUVAnimation : public FAssetTypeActions_Base
+{
+public:
+	// IAssetTypeActions Implementation
+	virtual FColor GetTypeColor() const override { return FColor(255,255,255); }
+	virtual uint32 GetCategories() override { return EAssetTypeCategories::Misc; }
+	virtual FText GetName() const override { return NSLOCTEXT("AssetTypeActions", "AssetTypeActions_SubUVAnimation", "Sub UV Animation"); }
+	virtual UClass* GetSupportedClass() const override { return USubUVAnimation::StaticClass(); }
+	// End IAssetTypeActions
+};
