Fixed linker error due to empty definition of AssetTypeActions_AnimStreamable::GetActions not being declared inline.

diff --git a/Engine/Source/Developer/AssetTools/Private/AssetTypeActions/AssetTypeActions_AnimStreamable.h b/Engine/Source/Developer/AssetTools/Private/AssetTypeActions/AssetTypeActions_AnimStreamable.h
--- a/Engine/Source/Developer/AssetTools/Private/AssetTypeActions/AssetTypeActions_AnimStreamable.h
+++ b/Engine/Source/Developer/AssetTools/Private/AssetTypeActions/AssetTypeActions_AnimStreamable.h
@@ -16,12 +16,12 @@ public:
 	virtual FColor GetTypeColor() const override { return FColor(181,230,29); }
 	virtual UClass* GetSupportedClass() const override { return UAnimStreamable::StaticClass(); }
 	virtual bool CanFilter() override { return true; }
-	virtual void GetActions(const TArray<UObject*>& InObjects, struct FToolMenuSection& Section) override;
+	virtual void GetActions(const TArray<UObject*>& InObjects, struct FToolMenuSection& Section) override {}
 };
 
-void FAssetTypeActions_AnimStreamable::GetActions(const TArray<UObject*>& InObjects, FToolMenuSection& Section)
+/*void FAssetTypeActions_AnimStreamable::GetActions(const TArray<UObject*>& InObjects, FToolMenuSection& Section)
 {
-	/*auto Sequences = GetTypedWeakObjectPtrs<UAnimStreamable>(InObjects);
+	auto Sequences = GetTypedWeakObjectPtrs<UAnimStreamable>(InObjects);
 
 	auto Logger = [Sequences]()
 	{
@@ -76,5 +76,5 @@ void FAssetTypeActions_AnimStreamable::GetActions(const TArray<UObject*>& InObje
 		NSLOCTEXT("AnimStreamable", "AnimSequence_ReimportWithNewSourceTooltip", "Reimport the selected sequence(s) from a new source file."),
 		FSlateIcon(FEditorStyle::GetStyleSetName(), "Persona.AssetActions.ReimportAnim"),
 		FUIAction(FExecuteAction::CreateLambda(Logger))
-	);*/
-}
\ No newline at end of file
+	);
+}*/
\ No newline at end of file
