Split up "Materials & Textures" built in asset category. All texture and render target assets are now under the "Texture" category and all material assets + subsurface profile under "Materials".

The category had grown significantly over time and splitting it up makes it easier to find things.

diff --git a/Engine/Plugins/Experimental/VirtualHeightfieldMesh/Source/VirtualHeightfieldMeshEditor/Private/HeightfieldMinMaxTextureAssetTypeActions.cpp b/Engine/Plugins/Experimental/VirtualHeightfieldMesh/Source/VirtualHeightfieldMeshEditor/Private/HeightfieldMinMaxTextureAssetTypeActions.cpp
--- a/Engine/Plugins/Experimental/VirtualHeightfieldMesh/Source/VirtualHeightfieldMeshEditor/Private/HeightfieldMinMaxTextureAssetTypeActions.cpp
+++ b/Engine/Plugins/Experimental/VirtualHeightfieldMesh/Source/VirtualHeightfieldMeshEditor/Private/HeightfieldMinMaxTextureAssetTypeActions.cpp
@@ -23,7 +23,7 @@ FColor FAssetTypeActions_HeightfieldMinMaxTexture::GetTypeColor() const
 
 uint32 FAssetTypeActions_HeightfieldMinMaxTexture::GetCategories()
 {
-	return EAssetTypeCategories::MaterialsAndTextures; 
+	return EAssetTypeCategories::Textures; 
 }
 
 #undef LOCTEXT_NAMESPACE
diff --git a/Engine/Plugins/Media/MediaPlayerEditor/Source/MediaPlayerEditor/Private/AssetTools/MediaTextureActions.cpp b/Engine/Plugins/Media/MediaPlayerEditor/Source/MediaPlayerEditor/Private/AssetTools/MediaTextureActions.cpp
--- a/Engine/Plugins/Media/MediaPlayerEditor/Source/MediaPlayerEditor/Private/AssetTools/MediaTextureActions.cpp
+++ b/Engine/Plugins/Media/MediaPlayerEditor/Source/MediaPlayerEditor/Private/AssetTools/MediaTextureActions.cpp
@@ -66,7 +66,7 @@ void FMediaTextureActions::GetActions(const TArray<UObject*>& InObjects, FToolMe
 
 uint32 FMediaTextureActions::GetCategories()
 {
-	return EAssetTypeCategories::MaterialsAndTextures | EAssetTypeCategories::Media;
+	return EAssetTypeCategories::Textures | EAssetTypeCategories::Media;
 }
 
 
diff --git a/Engine/Plugins/Media/MediaPlayerEditor/Source/MediaPlayerEditor/Private/Factories/MediaTextureFactoryNew.cpp b/Engine/Plugins/Media/MediaPlayerEditor/Source/MediaPlayerEditor/Private/Factories/MediaTextureFactoryNew.cpp
--- a/Engine/Plugins/Media/MediaPlayerEditor/Source/MediaPlayerEditor/Private/Factories/MediaTextureFactoryNew.cpp
+++ b/Engine/Plugins/Media/MediaPlayerEditor/Source/MediaPlayerEditor/Private/Factories/MediaTextureFactoryNew.cpp
@@ -35,7 +35,7 @@ UObject* UMediaTextureFactoryNew::FactoryCreateNew(UClass* InClass, UObject* InP
 
 uint32 UMediaTextureFactoryNew::GetMenuCategories() const
 {
-	return EAssetTypeCategories::Media | EAssetTypeCategories::MaterialsAndTextures;
+	return EAssetTypeCategories::Media | EAssetTypeCategories::Textures;
 }
 
 
diff --git a/Engine/Source/Developer/AssetTools/Private/AssetTools.cpp b/Engine/Source/Developer/AssetTools/Private/AssetTools.cpp
--- a/Engine/Source/Developer/AssetTools/Private/AssetTools.cpp
+++ b/Engine/Source/Developer/AssetTools/Private/AssetTools.cpp
@@ -202,13 +202,14 @@ UAssetToolsImpl::UAssetToolsImpl(const FObjectInitializer& ObjectInitializer)
 	// Register the built-in advanced categories
 	AllocatedCategoryBits.Add(TEXT("_BuiltIn_0"), FAdvancedAssetCategory(EAssetTypeCategories::Animation, LOCTEXT("AnimationAssetCategory", "Animation")));
 	AllocatedCategoryBits.Add(TEXT("_BuiltIn_1"), FAdvancedAssetCategory(EAssetTypeCategories::Blueprint, LOCTEXT("BlueprintAssetCategory", "Blueprints")));
-	AllocatedCategoryBits.Add(TEXT("_BuiltIn_2"), FAdvancedAssetCategory(EAssetTypeCategories::MaterialsAndTextures, LOCTEXT("MaterialAssetCategory", "Materials & Textures")));
+	AllocatedCategoryBits.Add(TEXT("_BuiltIn_2"), FAdvancedAssetCategory(EAssetTypeCategories::Materials, LOCTEXT("MaterialAssetCategory", "Materials")));
 	AllocatedCategoryBits.Add(TEXT("_BuiltIn_3"), FAdvancedAssetCategory(EAssetTypeCategories::Sounds, LOCTEXT("SoundAssetCategory", "Sounds")));
 	AllocatedCategoryBits.Add(TEXT("_BuiltIn_4"), FAdvancedAssetCategory(EAssetTypeCategories::Physics, LOCTEXT("PhysicsAssetCategory", "Physics")));
 	AllocatedCategoryBits.Add(TEXT("_BuiltIn_5"), FAdvancedAssetCategory(EAssetTypeCategories::UI, LOCTEXT("UserInterfaceAssetCategory", "User Interface")));
 	AllocatedCategoryBits.Add(TEXT("_BuiltIn_6"), FAdvancedAssetCategory(EAssetTypeCategories::Misc, LOCTEXT("MiscellaneousAssetCategory", "Miscellaneous")));
 	AllocatedCategoryBits.Add(TEXT("_BuiltIn_7"), FAdvancedAssetCategory(EAssetTypeCategories::Gameplay, LOCTEXT("GameplayAssetCategory", "Gameplay")));
 	AllocatedCategoryBits.Add(TEXT("_BuiltIn_8"), FAdvancedAssetCategory(EAssetTypeCategories::Media, LOCTEXT("MediaAssetCategory", "Media")));
+	AllocatedCategoryBits.Add(TEXT("_BuiltIn_9"), FAdvancedAssetCategory(EAssetTypeCategories::Textures, LOCTEXT("TextureAssetCategory", "Textures")));
 
 	EAssetTypeCategories::Type BlendablesCategoryBit = RegisterAdvancedAssetCategory(FName(TEXT("Blendables")), LOCTEXT("BlendablesAssetCategory", "Blendables"));
 	EAssetTypeCategories::Type FoliageCategoryBit = RegisterAdvancedAssetCategory(FName(TEXT("Foliage")), LOCTEXT("FoliageAssetCategory", "Foliage"));
diff --git a/Engine/Source/Developer/AssetTools/Private/AssetTypeActions/AssetTypeActions_MaterialFunction.h b/Engine/Source/Developer/AssetTools/Private/AssetTypeActions/AssetTypeActions_MaterialFunction.h
--- a/Engine/Source/Developer/AssetTools/Private/AssetTypeActions/AssetTypeActions_MaterialFunction.h
+++ b/Engine/Source/Developer/AssetTools/Private/AssetTypeActions/AssetTypeActions_MaterialFunction.h
@@ -20,7 +20,7 @@ public:
 	virtual bool HasActions ( const TArray<UObject*>& InObjects ) const override { return true; }
 	virtual void GetActions(const TArray<UObject*>& InObjects, struct FToolMenuSection& Section) override;
 	virtual void OpenAssetEditor( const TArray<UObject*>& InObjects, TSharedPtr<class IToolkitHost> EditWithinLevelEditor = TSharedPtr<IToolkitHost>() ) override;
-	virtual uint32 GetCategories() override { return EAssetTypeCategories::MaterialsAndTextures; }
+	virtual uint32 GetCategories() override { return EAssetTypeCategories::Materials; }
 	virtual class UThumbnailInfo* GetThumbnailInfo(UObject* Asset) const override;
 
 protected:
@@ -41,7 +41,7 @@ public:
 	virtual FColor GetTypeColor() const override { return FColor(0,175,175); }
 	virtual UClass* GetSupportedClass() const override;
 	virtual bool CanFilter() override;
-	virtual uint32 GetCategories() override { return EAssetTypeCategories::MaterialsAndTextures; }
+	virtual uint32 GetCategories() override { return EAssetTypeCategories::Materials; }
 
 private:
 	virtual void ExecuteNewMFI(TArray<TWeakObjectPtr<UMaterialFunctionInterface>> Objects) override;
@@ -57,7 +57,7 @@ public:
 	virtual FColor GetTypeColor() const override { return FColor(0,175,175); }
 	virtual UClass* GetSupportedClass() const override;
 	virtual bool CanFilter() override;
-	virtual uint32 GetCategories() override { return EAssetTypeCategories::MaterialsAndTextures; }
+	virtual uint32 GetCategories() override { return EAssetTypeCategories::Materials; }
 
 private:
 	virtual void ExecuteNewMFI(TArray<TWeakObjectPtr<UMaterialFunctionInterface>> Objects) override;
diff --git a/Engine/Source/Developer/AssetTools/Private/AssetTypeActions/AssetTypeActions_MaterialInterface.h b/Engine/Source/Developer/AssetTools/Private/AssetTypeActions/AssetTypeActions_MaterialInterface.h
--- a/Engine/Source/Developer/AssetTools/Private/AssetTypeActions/AssetTypeActions_MaterialInterface.h
+++ b/Engine/Source/Developer/AssetTools/Private/AssetTypeActions/AssetTypeActions_MaterialInterface.h
@@ -17,7 +17,7 @@ public:
 	virtual bool HasActions ( const TArray<UObject*>& InObjects ) const override { return true; }
 	virtual void GetActions(const TArray<UObject*>& InObjects, struct FToolMenuSection& Section) override;
 	virtual bool CanFilter() override { return false; }
-	virtual uint32 GetCategories() override { return EAssetTypeCategories::MaterialsAndTextures; }
+	virtual uint32 GetCategories() override { return EAssetTypeCategories::Materials; }
 	virtual class UThumbnailInfo* GetThumbnailInfo(UObject* Asset) const override;
 	virtual EThumbnailPrimType GetDefaultThumbnailPrimitiveType(UObject* Asset) const override;
 
diff --git a/Engine/Source/Developer/AssetTools/Private/AssetTypeActions/AssetTypeActions_MaterialParameterCollection.h b/Engine/Source/Developer/AssetTools/Private/AssetTypeActions/AssetTypeActions_MaterialParameterCollection.h
--- a/Engine/Source/Developer/AssetTools/Private/AssetTypeActions/AssetTypeActions_MaterialParameterCollection.h
+++ b/Engine/Source/Developer/AssetTools/Private/AssetTypeActions/AssetTypeActions_MaterialParameterCollection.h
@@ -13,5 +13,5 @@ public:
 	virtual FText GetName() const override { return NSLOCTEXT("AssetTypeActions", "AssetTypeActions_MaterialParameterCollection", "Material Parameter Collection"); }
 	virtual FColor GetTypeColor() const override { return FColor(0, 192, 0); }
 	virtual UClass* GetSupportedClass() const override { return UMaterialParameterCollection::StaticClass(); }
-	virtual uint32 GetCategories() override { return EAssetTypeCategories::MaterialsAndTextures; }
+	virtual uint32 GetCategories() override { return EAssetTypeCategories::Materials; }
 };
diff --git a/Engine/Source/Developer/AssetTools/Private/AssetTypeActions/AssetTypeActions_SubsurfaceProfile.h b/Engine/Source/Developer/AssetTools/Private/AssetTypeActions/AssetTypeActions_SubsurfaceProfile.h
--- a/Engine/Source/Developer/AssetTools/Private/AssetTypeActions/AssetTypeActions_SubsurfaceProfile.h
+++ b/Engine/Source/Developer/AssetTools/Private/AssetTypeActions/AssetTypeActions_SubsurfaceProfile.h
@@ -13,5 +13,5 @@ public:
 	virtual FColor GetTypeColor() const override { return FColor(255, 196, 128); }
 	virtual UClass* GetSupportedClass() const override;
 
-	virtual uint32 GetCategories() override { return EAssetTypeCategories::MaterialsAndTextures; }
+	virtual uint32 GetCategories() override { return EAssetTypeCategories::Materials; }
 };
diff --git a/Engine/Source/Developer/AssetTools/Private/AssetTypeActions/AssetTypeActions_Texture.h b/Engine/Source/Developer/AssetTools/Private/AssetTypeActions/AssetTypeActions_Texture.h
--- a/Engine/Source/Developer/AssetTools/Private/AssetTypeActions/AssetTypeActions_Texture.h
+++ b/Engine/Source/Developer/AssetTools/Private/AssetTypeActions/AssetTypeActions_Texture.h
@@ -19,7 +19,7 @@ public:
 	virtual void GetActions(const TArray<UObject*>& InObjects, struct FToolMenuSection& Section) override;
 	virtual void OpenAssetEditor( const TArray<UObject*>& InObjects, TSharedPtr<class IToolkitHost> EditWithinLevelEditor = TSharedPtr<IToolkitHost>() ) override;
 	virtual bool CanFilter() override { return false; }
-	virtual uint32 GetCategories() override { return EAssetTypeCategories::MaterialsAndTextures; }
+	virtual uint32 GetCategories() override { return EAssetTypeCategories::Textures; }
 	virtual bool IsImportedAsset() const override { return true; }
 	virtual void GetResolvedSourceFilePaths(const TArray<UObject*>& TypeAssets, TArray<FString>& OutSourceFilePaths) const override;
 
diff --git a/Engine/Source/Developer/AssetTools/Public/AssetTypeCategories.h b/Engine/Source/Developer/AssetTools/Public/AssetTypeCategories.h
--- a/Engine/Source/Developer/AssetTools/Public/AssetTypeCategories.h
+++ b/Engine/Source/Developer/AssetTools/Public/AssetTypeCategories.h
@@ -17,7 +17,7 @@ namespace EAssetTypeCategories
 		None = 0,
 		Basic = 1 << 0,
 		Animation = 1 << 1,
-		MaterialsAndTextures = 1 << 2,
+		Materials = 1 << 2,
 		Sounds = 1 << 3,
 		Physics = 1 << 4,
 		UI = 1 << 5,
@@ -25,9 +25,10 @@ namespace EAssetTypeCategories
 		Gameplay = 1 << 7,
 		Blueprint = 1 << 8,
 		Media = 1 << 9,
+		Textures = 1 << 10,
 
 		// Items below this will be allocated at runtime via RegisterAdvancedAssetCategory
-		FirstUser = 1 << 10,
+		FirstUser = 1 << 11,
 		LastUser = 1 << 31,
 		// Last allowed value is 1 << 31
 	};
diff --git a/Engine/Source/Editor/VirtualTexturingEditor/Private/RuntimeVirtualTextureAssetTypeActions.cpp b/Engine/Source/Editor/VirtualTexturingEditor/Private/RuntimeVirtualTextureAssetTypeActions.cpp
--- a/Engine/Source/Editor/VirtualTexturingEditor/Private/RuntimeVirtualTextureAssetTypeActions.cpp
+++ b/Engine/Source/Editor/VirtualTexturingEditor/Private/RuntimeVirtualTextureAssetTypeActions.cpp
@@ -216,7 +216,7 @@ FColor FAssetTypeActions_RuntimeVirtualTexture::GetTypeColor() const
 
 uint32 FAssetTypeActions_RuntimeVirtualTexture::GetCategories() 
 {
-	return EAssetTypeCategories::MaterialsAndTextures; 
+	return EAssetTypeCategories::Textures; 
 }
 
 void FAssetTypeActions_RuntimeVirtualTexture::GetActions(TArray<UObject*> const& InObjects, FMenuBuilder& MenuBuilder)
diff --git a/Engine/Source/Editor/VirtualTexturingEditor/Private/VirtualTextureBuilderAssetTypeActions.cpp b/Engine/Source/Editor/VirtualTexturingEditor/Private/VirtualTextureBuilderAssetTypeActions.cpp
--- a/Engine/Source/Editor/VirtualTexturingEditor/Private/VirtualTextureBuilderAssetTypeActions.cpp
+++ b/Engine/Source/Editor/VirtualTexturingEditor/Private/VirtualTextureBuilderAssetTypeActions.cpp
@@ -23,7 +23,7 @@ FColor FAssetTypeActions_VirtualTextureBuilder::GetTypeColor() const
 
 uint32 FAssetTypeActions_VirtualTextureBuilder::GetCategories() 
 {
-	return EAssetTypeCategories::MaterialsAndTextures; 
+	return EAssetTypeCategories::Textures; 
 }
 
 #undef LOCTEXT_NAMESPACE
