Unstable crash on moving assets from one folder to another

Referencing packages & possibly their non RF_Standalone content were GC'd at the first GC pass, before they had the chance of being saved.
Made sure referencing packages are saved before a GC pass is done.

diff --git a/Engine/Source/Developer/AssetTools/Private/AssetRenameManager.cpp b/Engine/Source/Developer/AssetTools/Private/AssetRenameManager.cpp
--- a/Engine/Source/Developer/AssetTools/Private/AssetRenameManager.cpp
+++ b/Engine/Source/Developer/AssetTools/Private/AssetRenameManager.cpp
@@ -472,10 +472,8 @@ bool FAssetRenameManager::FixReferencesAndRename(const TArray<FAssetRenameData>&
 			else
 			{
 				// Perform the rename, leaving redirectors only for assets which need them
-				PerformAssetRename(AssetsToRename);
-
-				// Save all packages that were referencing any of the assets that were moved without redirectors
-				SaveReferencingPackages(ReferencingPackagesToSave);
+				// Also save all packages that were referencing any of the assets that were moved without redirectors
+				PerformAssetRename(AssetsToRename, ReferencingPackagesToSave);
 
 				// Issue post rename event
 				AssetPostRenameEvent.Broadcast(AssetsAndNames);
@@ -1373,6 +1371,11 @@ bool FAssetRenameManager::CheckPackageForSoftObjectReferences(UPackage* Package,
 }
 
 void FAssetRenameManager::PerformAssetRename(TArray<FAssetRenameDataWithReferencers>& AssetsToRename) const
+{
+	PerformAssetRename(AssetsToRename, TArray<UPackage*>());
+}
+
+void FAssetRenameManager::PerformAssetRename(TArray<FAssetRenameDataWithReferencers>& AssetsToRename, const TArray<UPackage*>& ReferencingPackagesToSave) const
 {
 	const FText AssetRenameSlowTask = LOCTEXT("AssetRenameSlowTask", "Renaming Assets");
 	GWarn->BeginSlowTask(AssetRenameSlowTask, true);
@@ -1388,7 +1391,7 @@ void FAssetRenameManager::PerformAssetRename(TArray<FAssetRenameDataWithReferenc
 	FEditorFileUtils::GetDirtyWorldPackages(DirtyPackagesToCheckForSoftReferences);
 	FEditorFileUtils::GetDirtyContentPackages(DirtyPackagesToCheckForSoftReferences);
 
-	TArray<UPackage*> PackagesToSave;
+	TArray<UPackage*> PackagesToSave = ReferencingPackagesToSave;
 	TArray<UPackage*> PotentialPackagesToDelete;
 	for (int32 AssetIdx = 0; AssetIdx < AssetsToRename.Num(); ++AssetIdx)
 	{
diff --git a/Engine/Source/Developer/AssetTools/Private/AssetRenameManager.h b/Engine/Source/Developer/AssetTools/Private/AssetRenameManager.h
--- a/Engine/Source/Developer/AssetTools/Private/AssetRenameManager.h
+++ b/Engine/Source/Developer/AssetTools/Private/AssetRenameManager.h
@@ -110,6 +110,9 @@ private:
 	/** Performs the asset rename after the user has selected to proceed */
 	void PerformAssetRename(TArray<FAssetRenameDataWithReferencers>& AssetsToRename) const;
 
+	/** Performs the asset rename after the user has selected to proceed, also saving the provided referencing packages at the same time */
+	void PerformAssetRename(TArray<FAssetRenameDataWithReferencers>& AssetsToRename, const TArray<UPackage*>& ReferencingPackagesToSave) const;
+
 	/** Saves all the referencing packages and updates SCC state */
 	void SaveReferencingPackages(const TArray<UPackage*>& ReferencingPackagesToSave) const;
 
