Fix objects being garbage collected while executing a "Fixup Redirectors" for redirectors referenced by multiple maps.

Would potentially lead to external actors packages being treated as "empty", being deleted from disk & marked for deletion in the SCC.

diff --git a/Engine/Source/Developer/AssetTools/Private/AssetFixUpRedirectors.cpp b/Engine/Source/Developer/AssetTools/Private/AssetFixUpRedirectors.cpp
--- a/Engine/Source/Developer/AssetTools/Private/AssetFixUpRedirectors.cpp
+++ b/Engine/Source/Developer/AssetTools/Private/AssetFixUpRedirectors.cpp
@@ -98,6 +98,19 @@ void FAssetFixUpRedirectors::ExecuteFixUp(TArray<TWeakObjectPtr<UObjectRedirecto
 			TArray<UPackage*> ReferencingPackagesToSave;
 			LoadReferencingPackages(RedirectorRefsList, ReferencingPackagesToSave);
 
+			// Add all referencing packages objects that aren't RF_Standalone to the root set to avoid them being GC'd during the following processing
+			TSet<UObject*> RootedObjects;
+			for (UPackage* Package : ReferencingPackagesToSave)
+			{
+				ForEachObjectWithPackage(Package, [&RootedObjects](UObject* Object)
+				{
+					check(!Object->IsRooted());
+					Object->AddToRoot();
+					RootedObjects.Add(Object);
+					return true;
+				}, false, RF_Standalone, EInternalObjectFlags::RootSet);
+			}
+
 			// Check out all referencing packages, leave redirectors for assets referenced by packages that are not checked out and remove those packages from the save list.
 			const bool bUserAcceptedCheckout = CheckOutReferencingPackages(RedirectorRefsList, ReferencingPackagesToSave, bCheckoutDialogPrompt);
 			if ( bUserAcceptedCheckout )
@@ -124,6 +137,12 @@ void FAssetFixUpRedirectors::ExecuteFixUp(TArray<TWeakObjectPtr<UObjectRedirecto
 				// Finally, report any failures that happened during the rename
 				ReportFailures(RedirectorRefsList);
 			}
+
+			// Remove objects we manually rooted to allow them to be properly GC'd
+			for (UObject* Object : RootedObjects)
+			{
+				Object->RemoveFromRoot();
+			}
 		}
 	}
 }
