Added AssetTypeActions for UForceFeedbackAttenuation.

diff --git a/Engine/Source/Developer/AssetTools/Private/AssetTools.cpp b/Engine/Source/Developer/AssetTools/Private/AssetTools.cpp
--- a/Engine/Source/Developer/AssetTools/Private/AssetTools.cpp
+++ b/Engine/Source/Developer/AssetTools/Private/AssetTools.cpp
@@ -68,6 +68,7 @@
 #include "AssetTypeActions/AssetTypeActions_Struct.h"
 #include "AssetTypeActions/AssetTypeActions_Font.h"
 #include "AssetTypeActions/AssetTypeActions_FontFace.h"
+#include "AssetTypeActions/AssetTypeActions_ForceFeedbackAttenuation.h"
 #include "AssetTypeActions/AssetTypeActions_ForceFeedbackEffect.h"
 #include "AssetTypeActions/AssetTypeActions_HapticFeedback.h"
 #include "AssetTypeActions/AssetTypeActions_HLODProxy.h"
@@ -250,6 +251,7 @@ UAssetToolsImpl::UAssetToolsImpl(const FObjectInitializer& ObjectInitializer)
 	RegisterAssetTypeActions(MakeShareable(new FAssetTypeActions_SceneImportData));
 	RegisterAssetTypeActions(MakeShareable(new FAssetTypeActions_Font));
 	RegisterAssetTypeActions(MakeShareable(new FAssetTypeActions_FontFace));
+	RegisterAssetTypeActions(MakeShareable(new FAssetTypeActions_ForceFeedbackAttenuation));
 	RegisterAssetTypeActions(MakeShareable(new FAssetTypeActions_ForceFeedbackEffect));
 	RegisterAssetTypeActions(MakeShareable(new FAssetTypeActions_HLODProxy));
 	RegisterAssetTypeActions(MakeShareable(new FAssetTypeActions_SubsurfaceProfile));
diff --git a/Engine/Source/Developer/AssetTools/Private/AssetTypeActions/AssetTypeActions_ForceFeedbackAttenuation.h b/Engine/Source/Developer/AssetTools/Private/AssetTypeActions/AssetTypeActions_ForceFeedbackAttenuation.h
new file mode 100644
--- /dev/null
+++ b/Engine/Source/Developer/AssetTools/Private/AssetTypeActions/AssetTypeActions_ForceFeedbackAttenuation.h
@@ -0,0 +1,17 @@
+// Copyright Epic Games, Inc. All Rights Reserved.
+
+#pragma once
+
+#include "CoreMinimal.h"
+#include "AssetTypeActions_Base.h"
+#include "GameFramework/ForceFeedbackAttenuation.h"
+
+class FAssetTypeActions_ForceFeedbackAttenuation : public FAssetTypeActions_Base
+{
+public:
+	// IAssetTypeActions Implementation
+	virtual FText GetName() const override { return NSLOCTEXT("AssetTypeActions", "AssetTypeActions_ForceFeedbackAttenuation", "Force Feedback Attenuation"); }
+	virtual FColor GetTypeColor() const override { return FColor(175, 0, 0); }
+	virtual UClass* GetSupportedClass() const override { return UForceFeedbackAttenuation::StaticClass(); }
+	virtual uint32 GetCategories() override { return EAssetTypeCategories::Misc; }
+};
\ No newline at end of file
