Fix crash opening a WP map after a redirector fixup

diff --git a/Engine/Source/Developer/AssetTools/Private/AssetFixUpRedirectors.cpp b/Engine/Source/Developer/AssetTools/Private/AssetFixUpRedirectors.cpp
--- a/Engine/Source/Developer/AssetTools/Private/AssetFixUpRedirectors.cpp
+++ b/Engine/Source/Developer/AssetTools/Private/AssetFixUpRedirectors.cpp
@@ -96,7 +96,8 @@ void FAssetFixUpRedirectors::ExecuteFixUp(TArray<TWeakObjectPtr<UObjectRedirecto
 		{
 			// Load all referencing packages.
 			TArray<UPackage*> ReferencingPackagesToSave;
-			LoadReferencingPackages(RedirectorRefsList, ReferencingPackagesToSave);
+			TArray<UPackage*> LoadedPackages;
+			LoadReferencingPackages(RedirectorRefsList, ReferencingPackagesToSave, LoadedPackages);
 
 			// Add all referencing packages objects that aren't RF_Standalone to the root set to avoid them being GC'd during the following processing
 			TSet<UObject*> RootedObjects;
@@ -143,6 +144,22 @@ void FAssetFixUpRedirectors::ExecuteFixUp(TArray<TWeakObjectPtr<UObjectRedirecto
 			{
 				Object->RemoveFromRoot();
 			}
+
+			// If any packages were loaded during the fixup process, make sure we unload them here
+			if (!LoadedPackages.IsEmpty())
+			{
+				for (UPackage* LoadedPackage : LoadedPackages)
+				{
+					ForEachObjectWithPackage(LoadedPackage, [](UObject* Object)
+					{
+						Object->ClearFlags(RF_Standalone);
+						return true;
+					}, false);
+				}
+
+				// Collect garbage.
+				CollectGarbage(GARBAGE_COLLECTION_KEEPFLAGS);
+			}
 		}
 	}
 }
@@ -185,7 +202,7 @@ bool FAssetFixUpRedirectors::UpdatePackageStatus(const TArray<FRedirectorRefs>&
 	return true;
 }
 
-void FAssetFixUpRedirectors::LoadReferencingPackages(TArray<FRedirectorRefs>& RedirectorsToFix, TArray<UPackage*>& OutReferencingPackagesToSave) const
+void FAssetFixUpRedirectors::LoadReferencingPackages(TArray<FRedirectorRefs>& RedirectorsToFix, TArray<UPackage*>& OutReferencingPackagesToSave, TArray<UPackage*>& OutLoadedPackages) const
 {
 	FScopedSlowTask SlowTask( RedirectorsToFix.Num(), LOCTEXT( "LoadingReferencingPackages", "Loading Referencing Packages..." ) );
 	SlowTask.MakeDialog();
@@ -220,6 +237,10 @@ void FAssetFixUpRedirectors::LoadReferencingPackages(TArray<FRedirectorRefs>& Re
 			if ( !Package )
 			{
 				Package = LoadPackage(NULL, *PackageName, LOAD_None);
+				if (Package)
+				{
+					OutLoadedPackages.Add(Package);
+				}
 			}
 
 			if ( Package )
diff --git a/Engine/Source/Developer/AssetTools/Private/AssetFixUpRedirectors.h b/Engine/Source/Developer/AssetTools/Private/AssetFixUpRedirectors.h
--- a/Engine/Source/Developer/AssetTools/Private/AssetFixUpRedirectors.h
+++ b/Engine/Source/Developer/AssetTools/Private/AssetFixUpRedirectors.h
@@ -33,7 +33,7 @@ private:
 	  * Loads all referencing packages to redirectors in RedirectorsToFix, finds redirectors whose references can
 	  * not be fixed up, and returns a list of referencing packages to save.
 	  */
-	void LoadReferencingPackages(TArray<FRedirectorRefs>& RedirectorsToFix, TArray<UPackage*>& OutReferencingPackagesToSave) const;
+	void LoadReferencingPackages(TArray<FRedirectorRefs>& RedirectorsToFix, TArray<UPackage*>& OutReferencingPackagesToSave, TArray<UPackage*>& OutLoadedPackages)const;
 
 	/** 
 	  * Check out referencing packages and marks assets whose referencing packages were not checked out to not fix the redirector.
