Fix Notification Message on Cancel migration

diff --git a/Engine/Source/Developer/AssetTools/Private/AssetTools.cpp b/Engine/Source/Developer/AssetTools/Private/AssetTools.cpp
--- a/Engine/Source/Developer/AssetTools/Private/AssetTools.cpp
+++ b/Engine/Source/Developer/AssetTools/Private/AssetTools.cpp
@@ -3005,6 +3005,10 @@ void UAssetToolsImpl::MigratePackages_ReportConfirmed(TSharedPtr<TArray<ReportPa
 		}
 		LogMessage = FText::FromString(ErrorMessage);
 	}
+	else if (bUserCanceled)
+	{
+		LogMessage = LOCTEXT("MigratePackages_CanceledPage", "Content migration was canceled.");
+	}
 	else if ( !bUserCanceled )
 	{
 		MigrateLog.NewPage( LOCTEXT("MigratePackages_CompletePage", "Content migration completed successfully!") );
