Optimize FAssetRenameDataWithReferencers by changing an array with a set.
Results for breaking a Level Instance with ~10K actors: ~810s down to ~256s.

diff --git a/Engine/Source/Developer/AssetTools/Private/AssetRenameManager.cpp b/Engine/Source/Developer/AssetTools/Private/AssetRenameManager.cpp
--- a/Engine/Source/Developer/AssetTools/Private/AssetRenameManager.cpp
+++ b/Engine/Source/Developer/AssetTools/Private/AssetRenameManager.cpp
@@ -87,7 +87,7 @@ namespace AssetRenameManagerImpl
 
 struct FAssetRenameDataWithReferencers : public FAssetRenameData
 {
-	TArray<FName> ReferencingPackageNames;
+	TSet<FName> ReferencingPackageNames;
 	FText FailureReason;
 	bool bCreateRedirector;
 	bool bRenameFailed;
@@ -780,19 +780,19 @@ void FAssetRenameManager::PopulateAssetReferencers(TArray<FAssetRenameDataWithRe
 		{
 			if (!RenamingAssetPackageNames.Contains(ReferencingPackageName))
 			{
-				AssetToRename.ReferencingPackageNames.AddUnique(ReferencingPackageName);
+				AssetToRename.ReferencingPackageNames.Add(ReferencingPackageName);
 			}
 		}
 
 		if (AssetToRename.bOnlyFixSoftReferences)
 		{
-			AssetToRename.ReferencingPackageNames.AddUnique(FName(*AssetToRename.OldObjectPath.GetLongPackageName()));
-			AssetToRename.ReferencingPackageNames.AddUnique(FName(*AssetToRename.NewObjectPath.GetLongPackageName()));
+			AssetToRename.ReferencingPackageNames.Add(FName(*AssetToRename.OldObjectPath.GetLongPackageName()));
+			AssetToRename.ReferencingPackageNames.Add(FName(*AssetToRename.NewObjectPath.GetLongPackageName()));
 
 			// Add dirty packages and the package that owns the reference. They will get filtered out in LoadReferencingPackages if they aren't valid
 			for (UPackage* Package : ExtraPackagesToCheckForSoftReferences)
 			{
-				AssetToRename.ReferencingPackageNames.AddUnique(Package->GetFName());
+				AssetToRename.ReferencingPackageNames.Add(Package->GetFName());
 			}
 		}
 	}
@@ -914,9 +914,10 @@ void FAssetRenameManager::LoadReferencingPackages(TArray<FAssetRenameDataWithRef
 
 		TArray<UPackage*> PackagesToSaveForThisAsset;
 		bool bAllPackagesLoadedForThisAsset = true;
-		for (int32 i = 0; i < RenameData.ReferencingPackageNames.Num(); i++)
+
+		for (auto It = RenameData.ReferencingPackageNames.CreateIterator(); It; ++It)
 		{
-			FName PackageName = RenameData.ReferencingPackageNames[i];
+			FName PackageName = *It;
 			// Check if the package is a map before loading it!
 			if (!bLoadAllPackages && FEditorFileUtils::IsMapPackageAsset(PackageName.ToString()))
 			{
@@ -953,8 +954,7 @@ void FAssetRenameManager::LoadReferencingPackages(TArray<FAssetRenameDataWithRef
 				else
 				{
 					// This package does not actually reference the asset, so remove it
-					RenameData.ReferencingPackageNames.RemoveAt(i);
-					i--;
+					It.RemoveCurrent();
 				}
 			}
 			else
