DataAssets - DataAssets can now have their base class changed (single or in bulk) as a right click asset action.  So if you refactor your data assets and make new subclasses it's a lot easier to upgrade existing content now.

diff --git a/Engine/Source/Developer/AssetTools/Private/AssetTypeActions/AssetTypeActions_DataAsset.cpp b/Engine/Source/Developer/AssetTools/Private/AssetTypeActions/AssetTypeActions_DataAsset.cpp
--- a/Engine/Source/Developer/AssetTools/Private/AssetTypeActions/AssetTypeActions_DataAsset.cpp
+++ b/Engine/Source/Developer/AssetTools/Private/AssetTypeActions/AssetTypeActions_DataAsset.cpp
@@ -1,11 +1,43 @@
 // Copyright Epic Games, Inc. All Rights Reserved.
 
 #include "AssetTypeActions/AssetTypeActions_DataAsset.h"
+#include "Engine/Engine.h"
+#include "ToolMenuSection.h"
+#include "ClassViewerModule.h"
+#include "Kismet2/SClassPickerDialog.h"
+#include "ClassViewerFilter.h"
+#include "ObjectTools.h"
 
 #define LOCTEXT_NAMESPACE "AssetTypeActions"
 
 static const FName NAME_NativeClass(TEXT("NativeClass"));
 
+class FNewNodeClassFilter : public IClassViewerFilter
+{
+public:
+	FNewNodeClassFilter(UClass* InBaseClass)
+		: BaseClass(InBaseClass)
+	{
+	}
+
+	virtual bool IsClassAllowed(const FClassViewerInitializationOptions& InInitOptions, const UClass* InClass, TSharedRef< FClassViewerFilterFuncs > InFilterFuncs) override
+	{
+		if (InClass != nullptr)
+		{
+			return InClass->IsChildOf(BaseClass);
+		}
+		return false;
+	}
+
+	virtual bool IsUnloadedClassAllowed(const FClassViewerInitializationOptions& InInitOptions, const TSharedRef< const IUnloadedBlueprintData > InUnloadedClassData, TSharedRef< FClassViewerFilterFuncs > InFilterFuncs) override
+	{
+		return InUnloadedClassData->IsChildOf(BaseClass);
+	}
+
+private:
+	UClass* BaseClass;
+};
+
 FText FAssetTypeActions_DataAsset::GetDisplayNameFromAssetData(const FAssetData& AssetData) const
 {
 	if (AssetData.IsValid())
@@ -23,4 +55,60 @@ FText FAssetTypeActions_DataAsset::GetDisplayNameFromAssetData(const FAssetData&
 	return FText::GetEmpty();
 }
 
+void FAssetTypeActions_DataAsset::GetActions(const TArray<UObject*>& InObjects, FToolMenuSection& Section)
+{
+	TArray<TWeakObjectPtr<UDataAsset>> DataAssets = GetTypedWeakObjectPtrs<UDataAsset>(InObjects);
+
+	Section.AddMenuEntry(
+		"DataAsset_ChangeClass",
+		LOCTEXT("DataAsset_ChangeClass", "Change DataAsset Class"),
+		LOCTEXT("DataAsset_ChangeClassTip", "Change the class these Data Assets are subclassed from."),
+		FSlateIcon(),
+		FUIAction(
+			FExecuteAction::CreateSP(this, &FAssetTypeActions_DataAsset::ExecuteChangeDataAssetClass, DataAssets)
+		)
+	);
+}
+
+void FAssetTypeActions_DataAsset::ExecuteChangeDataAssetClass(TArray<TWeakObjectPtr<UDataAsset>> InDataAssets)
+{
+	const FText TitleText = LOCTEXT("DataAsset_PickNewDataAssetClass", "Pick New DataAsset Class");
+	FClassViewerInitializationOptions Options;
+	Options.ClassFilters.Add(MakeShared<FNewNodeClassFilter>(UDataAsset::StaticClass()));
+	UClass* OutNewDataAssetClass = nullptr;
+	const bool bPressedOk = SClassPickerDialog::PickClass(TitleText, Options, OutNewDataAssetClass, UDataAsset::StaticClass());
+
+	if (bPressedOk && OutNewDataAssetClass != nullptr)
+	{
+		for (TWeakObjectPtr<UDataAsset> DataAssetPtr : InDataAssets)
+		{
+			if (UDataAsset* OldDataAsset = DataAssetPtr.Get())
+			{
+				if (OldDataAsset && OldDataAsset->IsValidLowLevel())
+				{
+					FName ObjectName = OldDataAsset->GetFName();
+					UObject* Outer = OldDataAsset->GetOuter();
+					OldDataAsset->Rename(nullptr, GetTransientPackage(), REN_DoNotDirty | REN_DontCreateRedirectors);
+
+					UObject* NewDataAsset = NewObject<UObject>(Outer, OutNewDataAssetClass, ObjectName, OldDataAsset->GetFlags());
+
+					// Migrate Data
+					{
+						UEngine::FCopyPropertiesForUnrelatedObjectsParams CopyOptions;
+						CopyOptions.bNotifyObjectReplacement = true;
+						UEngine::CopyPropertiesForUnrelatedObjects(OldDataAsset, NewDataAsset, CopyOptions);
+					}
+
+					NewDataAsset->MarkPackageDirty();
+
+					// Consolidate or "Replace" the old object with the new object for any living references.
+					bool bShowDeleteConfirmation = false;
+					TArray<UObject*> OldDataAssetArray({ (UObject*)OldDataAsset });
+					ObjectTools::ConsolidateObjects(NewDataAsset, OldDataAssetArray, bShowDeleteConfirmation);
+				}
+			}
+		}
+	}
+}
+
 #undef LOCTEXT_NAMESPACE
diff --git a/Engine/Source/Developer/AssetTools/Public/AssetTypeActions/AssetTypeActions_DataAsset.h b/Engine/Source/Developer/AssetTools/Public/AssetTypeActions/AssetTypeActions_DataAsset.h
new file mode 100644
--- /dev/null
+++ b/Engine/Source/Developer/AssetTools/Public/AssetTypeActions/AssetTypeActions_DataAsset.h
@@ -0,0 +1,24 @@
+// Copyright Epic Games, Inc. All Rights Reserved.
+
+#pragma once
+
+#include "CoreMinimal.h"
+#include "AssetTypeActions_Base.h"
+#include "Engine/DataAsset.h"
+#include "UObject/WeakObjectPtrTemplates.h"
+
+class ASSETTOOLS_API FAssetTypeActions_DataAsset : public FAssetTypeActions_Base
+{
+public:
+	// IAssetTypeActions Implementation
+	virtual FText GetName() const override { return NSLOCTEXT("AssetTypeActions", "AssetTypeActions_DataAsset", "Data Asset"); }
+	virtual FColor GetTypeColor() const override { return FColor(201, 29, 85); }
+	virtual UClass* GetSupportedClass() const override { return UDataAsset::StaticClass(); }
+	virtual uint32 GetCategories() override { return EAssetTypeCategories::Misc; }
+	virtual FText GetDisplayNameFromAssetData(const FAssetData& AssetData) const override;
+	virtual bool HasActions(const TArray<UObject*>& InObjects) const override { return true; }
+	virtual void GetActions(const TArray<UObject*>& InObjects, struct FToolMenuSection& Section) override;
+
+private:
+	void ExecuteChangeDataAssetClass(TArray<TWeakObjectPtr<UDataAsset>> InDataAssets);
+};
