DataAsset - Change the wording some on the action to convert them to another subclass.

diff --git a/Engine/Source/Developer/AssetTools/Private/AssetTypeActions/AssetTypeActions_DataAsset.cpp b/Engine/Source/Developer/AssetTools/Private/AssetTypeActions/AssetTypeActions_DataAsset.cpp
--- a/Engine/Source/Developer/AssetTools/Private/AssetTypeActions/AssetTypeActions_DataAsset.cpp
+++ b/Engine/Source/Developer/AssetTools/Private/AssetTypeActions/AssetTypeActions_DataAsset.cpp
@@ -61,7 +61,7 @@ void FAssetTypeActions_DataAsset::GetActions(const TArray<UObject*>& InObjects,
 
 	Section.AddMenuEntry(
 		"DataAsset_ChangeClass",
-		LOCTEXT("DataAsset_ChangeClass", "Change DataAsset Class"),
+		LOCTEXT("DataAsset_ChangeClass", "Convert to different DataAsset type"),
 		LOCTEXT("DataAsset_ChangeClassTip", "Change the class these Data Assets are subclassed from."),
 		FSlateIcon(),
 		FUIAction(
