Control Rig: Ported change "Move create control rig menu item under create in content browser asset context menu" to UE5

diff --git a/Engine/Plugins/Animation/ControlRig/Source/ControlRigEditor/Private/ControlRigBlueprintActions.cpp b/Engine/Plugins/Animation/ControlRig/Source/ControlRigEditor/Private/ControlRigBlueprintActions.cpp
--- a/Engine/Plugins/Animation/ControlRig/Source/ControlRigEditor/Private/ControlRigBlueprintActions.cpp
+++ b/Engine/Plugins/Animation/ControlRig/Source/ControlRigEditor/Private/ControlRigBlueprintActions.cpp
@@ -83,8 +83,8 @@ TSharedPtr<SWidget> FControlRigBlueprintActions::GetThumbnailOverlay(const FAsse
 void FControlRigBlueprintActions::ExtendSketalMeshToolMenu()
 {
 	TArray<UToolMenu*> MenusToExtend;
-	MenusToExtend.Add(UToolMenus::Get()->ExtendMenu("ContentBrowser.AssetContextMenu.SkeletalMesh"));
-	MenusToExtend.Add(UToolMenus::Get()->ExtendMenu("ContentBrowser.AssetContextMenu.Skeleton"));
+	MenusToExtend.Add(UToolMenus::Get()->ExtendMenu("ContentBrowser.AssetContextMenu.SkeletalMesh.CreateSkeletalMeshSubmenu"));
+	MenusToExtend.Add(UToolMenus::Get()->ExtendMenu("ContentBrowser.AssetContextMenu.Skeleton.CreateSkeletalMeshSubmenu"));
 
 	for(UToolMenu* Menu : MenusToExtend)
 	{
@@ -93,7 +93,7 @@ void FControlRigBlueprintActions::ExtendSketalMeshToolMenu()
 			continue;
 		}
 
-		FToolMenuSection& Section = Menu->FindOrAddSection("GetAssetActions");
+		FToolMenuSection& Section = Menu->AddSection("ControlRig", LOCTEXT("ControlRigSectionName", "Control Rig"));
 		Section.AddDynamicEntry("GetActions", FNewToolMenuSectionDelegate::CreateLambda([](FToolMenuSection& InSection)
 		{
 			UContentBrowserAssetContextMenuContext* Context = InSection.FindContext<UContentBrowserAssetContextMenuContext>();
diff --git a/Engine/Source/Developer/AssetTools/Private/AssetTypeActions/AssetTypeActions_SkeletalMesh.cpp b/Engine/Source/Developer/AssetTools/Private/AssetTypeActions/AssetTypeActions_SkeletalMesh.cpp
--- a/Engine/Source/Developer/AssetTools/Private/AssetTypeActions/AssetTypeActions_SkeletalMesh.cpp
+++ b/Engine/Source/Developer/AssetTools/Private/AssetTypeActions/AssetTypeActions_SkeletalMesh.cpp
@@ -461,7 +461,7 @@ void FAssetTypeActions_SkeletalMesh::GetActions(const TArray<UObject*>& InObject
 			"CreateSkeletalMeshSubmenu",
 			LOCTEXT("CreateSkeletalMeshSubmenu", "Create"),
 			LOCTEXT("CreateSkeletalMeshSubmenu_ToolTip", "Create related assets"),
-			FNewMenuDelegate::CreateSP(this, &FAssetTypeActions_SkeletalMesh::FillCreateMenu, Meshes),
+			FNewToolMenuDelegate::CreateSP(this, &FAssetTypeActions_SkeletalMesh::FillCreateMenu, Meshes),
 			false,
 			FSlateIcon(FEditorStyle::GetStyleSetName(), "Persona.AssetActions.CreateAnimAsset")
 			);
@@ -491,20 +491,23 @@ void FAssetTypeActions_SkeletalMesh::GetActions(const TArray<UObject*>& InObject
 	);
 }
 
-void FAssetTypeActions_SkeletalMesh::FillCreateMenu(FMenuBuilder& MenuBuilder, TArray<TWeakObjectPtr<USkeletalMesh>> Meshes) const
+void FAssetTypeActions_SkeletalMesh::FillCreateMenu(UToolMenu* Menu, TArray<TWeakObjectPtr<USkeletalMesh>> Meshes) const
 {
-	MenuBuilder.BeginSection("CreatePhysicsAsset", LOCTEXT("CreatePhysicsAssetMenuHeading", "Physics Asset"));
 	{
-		MenuBuilder.AddSubMenu(
+		FToolMenuSection& Section = Menu->AddSection("CreatePhysicsAsset", LOCTEXT("CreatePhysicsAssetMenuHeading", "Physics Asset"));
+		Section.AddSubMenu(
+            "NewPhysicsAssetMenu",
 			LOCTEXT("SkeletalMesh_NewPhysicsAssetMenu", "Physics Asset"),
 			LOCTEXT("SkeletalMesh_NewPhysicsAssetMenu_ToolTip", "Options for creating new physics assets from the selected meshes."),
 			FNewMenuDelegate::CreateSP(const_cast<FAssetTypeActions_SkeletalMesh*>(this), &FAssetTypeActions_SkeletalMesh::GetPhysicsAssetMenu, Meshes));
 	}
-	MenuBuilder.EndSection();
 
 	TArray<TWeakObjectPtr<UObject>> Objects;
 	Algo::Transform(Meshes, Objects, [](const TWeakObjectPtr<USkeletalMesh>& SkelMesh) { return SkelMesh; });
-	AnimationEditorUtils::FillCreateAssetMenu(MenuBuilder, Objects, FAnimAssetCreated::CreateSP(this, &FAssetTypeActions_SkeletalMesh::OnAssetCreated));
+	Menu->AddDynamicSection(NAME_None, FNewToolMenuDelegateLegacy::CreateLambda([this, Objects](FMenuBuilder& MenuBuilder, UToolMenu* Menu)
+	{
+		AnimationEditorUtils::FillCreateAssetMenu(MenuBuilder, Objects, FAnimAssetCreated::CreateSP(this, &FAssetTypeActions_SkeletalMesh::OnAssetCreated));
+	}));
 }
 
 void FAssetTypeActions_SkeletalMesh::OpenAssetEditor( const TArray<UObject*>& InObjects, TSharedPtr<IToolkitHost> EditWithinLevelEditor )
diff --git a/Engine/Source/Developer/AssetTools/Private/AssetTypeActions/AssetTypeActions_SkeletalMesh.h b/Engine/Source/Developer/AssetTools/Private/AssetTypeActions/AssetTypeActions_SkeletalMesh.h
--- a/Engine/Source/Developer/AssetTools/Private/AssetTypeActions/AssetTypeActions_SkeletalMesh.h
+++ b/Engine/Source/Developer/AssetTools/Private/AssetTypeActions/AssetTypeActions_SkeletalMesh.h
@@ -69,7 +69,7 @@ private:
 
 	void FillSourceMenu(FMenuBuilder& MenuBuilder, TArray<TWeakObjectPtr<USkeletalMesh>> Meshes) const;
 	void FillSkeletonMenu(FMenuBuilder& MenuBuilder, TArray<TWeakObjectPtr<USkeletalMesh>> Meshes) const;
-	void FillCreateMenu(FMenuBuilder& MenuBuilder, TArray<TWeakObjectPtr<USkeletalMesh>> Meshes) const;
+	void FillCreateMenu(UToolMenu* MenuBuilder, TArray<TWeakObjectPtr<USkeletalMesh>> Meshes) const;
 
 	void OnAssetRemoved(const struct FAssetData& AssetData);
 
