Made the asset tools able to handle a invalid package path or to convert it to a package path when possible.

diff --git a/Engine/Source/Developer/AssetTools/Private/AssetTools.cpp b/Engine/Source/Developer/AssetTools/Private/AssetTools.cpp
--- a/Engine/Source/Developer/AssetTools/Private/AssetTools.cpp
+++ b/Engine/Source/Developer/AssetTools/Private/AssetTools.cpp
@@ -1959,7 +1959,23 @@ TArray<UObject*> UAssetToolsImpl::ImportAssetsInternal(const TArray<FString>& Fi
 		// can be created out of sources that contain spaces and other invalid characters. Filename cannot be sanitized
 		// until other checks are done that rely on looking at the actual source file so sanitation is delayed.
 		const FString& Filename = FilesAndDestinations[FileIdx].Key;
-		const FString DestinationPath = ObjectTools::SanitizeObjectPath(FilesAndDestinations[FileIdx].Value);
+
+		FString DestinationPath;
+		FString ErrorMsg;
+		if (!FPackageName::TryConvertFilenameToLongPackageName(ObjectTools::SanitizeObjectPath(FilesAndDestinations[FileIdx].Value), DestinationPath, &ErrorMsg))
+		{
+			const FText Message = FText::Format(LOCTEXT("CannotConvertDestinationPath", "Can't import the file '{0}' because the destination path '{1}' cannot be converted to a package path."), FText::FromString(Filename), FText::FromString(DestinationPath));
+			if (!bAutomatedImport)
+			{
+				FMessageDialog::Open(EAppMsgType::Ok, Message);
+			}
+
+			UE_LOG(LogAssetTools, Warning, TEXT("%s"), *ErrorMsg);
+			UE_LOG(LogAssetTools, Warning, TEXT("%s"), *Message.ToString());
+
+			continue;
+		}
+
 		if (bUseInterchangeFramework)
 		{
 			UE::Interchange::FScopedSourceData ScopedSourceData(Filename);
