Advanced Copy - Fixed folders not being copyable, because their default state was unchecked, this has been fixed to check the contents of the folders.
Advanced Copy - We no longer open every package and all dependencies to setup possible renames, the code has been changed to do all the logic on PackageName, needed to cut out bGenerateUniqueNames from the copy params for now.

TODO - A good future improvement would be to make it so that tracking dependencies becomes lazy in the UI, as someone expands items in the tree, it would extend the dependencies 2 levels to know if it termininates or not, and wait until it's expanded to further look for depependencies.  Would require a much larger refactor than I wanted to do here.

diff --git a/Engine/Source/Developer/AssetTools/Private/AssetTools.cpp b/Engine/Source/Developer/AssetTools/Private/AssetTools.cpp
--- a/Engine/Source/Developer/AssetTools/Private/AssetTools.cpp
+++ b/Engine/Source/Developer/AssetTools/Private/AssetTools.cpp
@@ -716,57 +716,77 @@ void UAssetToolsImpl::GenerateAdvancedCopyDestinations(FAdvancedCopyParams& InPa
 
 			if (bFileOKToCopy)
 			{
-				UPackage* Pkg = LoadPackage(nullptr, *PackageNameString, LOAD_None);
-				if (Pkg)
+				FString Parent = FString();
+				if (bGenerateRelativePaths)
 				{
-					FString Name = ObjectTools::SanitizeObjectName(FPaths::GetBaseFilename(SrcFilename));
-					UObject* ExistingObject = StaticFindObject(UObject::StaticClass(), Pkg, *Name);
-					if (ExistingObject)
+					FString RootFolder = UAdvancedCopyCustomization::StaticClass()->GetDefaultObject<UAdvancedCopyCustomization>()->GetPackageThatInitiatedCopy();
+					if (RootFolder != PackageNameString)
 					{
-						FString Parent = FString();
-						if (bGenerateRelativePaths)
+						FString BaseParent = FString();
+						int32 MinLength = RootFolder.Len() < PackageNameString.Len() ? RootFolder.Len() : PackageNameString.Len();
+						for (int Char = 0; Char < MinLength; Char++)
 						{
-							FString RootFolder = UAdvancedCopyCustomization::StaticClass()->GetDefaultObject<UAdvancedCopyCustomization>()->GetPackageThatInitiatedCopy();
-							if (RootFolder != FPaths::GetPath(ExistingObject->GetPathName()))
+							if (RootFolder[Char] == PackageNameString[Char])
 							{
-								FString BaseParent = FString();
-								int32 MinLength = RootFolder.Len() < PackageNameString.Len() ? RootFolder.Len() : PackageNameString.Len();
-								for (int Char = 0; Char < MinLength; Char++)
-								{
-									if (RootFolder[Char] == PackageNameString[Char])
-									{
-										BaseParent += RootFolder[Char];
-									}
-									else
-									{
-										break;
-									}
-								}
-								// If we are in the root content folder, don't break down the folder string
-								if (BaseParent == TEXT("/Game"))
-								{
-									Parent = BaseParent;
-								}
-								else
-								{
-									BaseParent.Split(TEXT("/"), &Parent, nullptr, ESearchCase::IgnoreCase, ESearchDir::FromEnd);
-								}
+								BaseParent += RootFolder[Char];
+							}
+							else
+							{
+								break;
 							}
 						}
 
-						ObjectTools::FMoveDialogInfo MoveDialogInfo;
-						MoveDialogInfo.bOkToAll = InParams.bCopyOverAllDestinationOverlaps;
-						MoveDialogInfo.bPromptForRenameOnConflict = false;
-						// The default value for save packages is true if SCC is enabled because the user can use SCC to revert a change
-						MoveDialogInfo.bSavePackages = ISourceControlModule::Get().IsEnabled();
-						ObjectTools::GetMoveDialogInfo(NSLOCTEXT("UnrealEd", "DuplicateObjects", "Copy Objects"), ExistingObject, InParams.bGenerateUniqueNames, Parent, DestinationFolder, MoveDialogInfo);
-						OutPackagesAndDestinations.Add(PackageNameString, MoveDialogInfo.PGN.PackageName);
+						// If we are in the root content folder, don't break down the folder string
+						if (BaseParent == TEXT("/Game"))
+						{
+							Parent = BaseParent;
+						}
+						else
+						{
+							BaseParent.Split(TEXT("/"), &Parent, nullptr, ESearchCase::IgnoreCase, ESearchDir::FromEnd);
+						}
 					}
 				}
+
+				const FString DestinationPackageName = UAssetToolsImpl::GenerateAdvancedCopyDestinationPackageName(PackageNameString, Parent, DestinationFolder);
+				OutPackagesAndDestinations.Add(PackageNameString, DestinationPackageName);
 			}
 		}
 	}
-	
+}
+
+FString UAssetToolsImpl::GenerateAdvancedCopyDestinationPackageName(const FString& SourcePackage, const FString& SourcePath, const FString& DestinationFolder)
+{
+	FString DestinationPackageName;
+
+	const bool bIsRelativeOperation = SourcePath.Len() && DestinationFolder.Len() && SourcePackage.StartsWith(SourcePath);
+	if (bIsRelativeOperation)
+	{
+		// Folder copy/move.
+
+		// Collect the relative path then use it to determine the new location
+		// For example, if SourcePath = /Game/MyPath and SourcePackage = /Game/MyPath/MySubPath/MyAsset
+		//     /Game/MyPath/MySubPath/MyAsset -> /MySubPath/
+
+		const int32 ShortPackageNameLen = FPackageName::GetShortName(SourcePackage).Len();
+		const int32 RelativePathLen = SourcePackage.Len() - ShortPackageNameLen - SourcePath.Len();
+		const FString RelativeDestPath = SourcePackage.Mid(SourcePath.Len(), RelativePathLen);
+
+		DestinationPackageName = DestinationFolder + RelativeDestPath + FPackageName::GetShortName(SourcePackage);
+	}
+	else if (DestinationFolder.Len())
+	{
+		// Use the passed in default path
+		// Normal path
+		DestinationPackageName = DestinationFolder + "/" + FPackageName::GetShortName(SourcePackage);
+	}
+	else
+	{
+		// Use the path from the old package
+		DestinationPackageName = SourcePackage;
+	}
+
+	return DestinationPackageName;
 }
 
 bool UAssetToolsImpl::FlattenAdvancedCopyDestinations(const TArray<TMap<FString, FString>> PackagesAndDestinations, TMap<FString, FString>& FlattenedPackagesAndDestinations) const
@@ -3178,9 +3198,8 @@ void UAssetToolsImpl::BeginAdvancedCopyPackages(const TArray<FName>& InputNamesT
 	const bool bSaveContentPackages = true;
 	if (FEditorFileUtils::SaveDirtyPackages(bPromptUserToSave, bSaveMapPackages, bSaveContentPackages))
 	{
-
-		FAssetRegistryModule& AssetRegistryModule = FModuleManager::Get().LoadModuleChecked<FAssetRegistryModule>(TEXT("AssetRegistry"));
-		if (AssetRegistryModule.Get().IsLoadingAssets())
+		IAssetRegistry& AssetRegistry = FAssetRegistryModule::GetRegistry();
+		if (AssetRegistry.IsLoadingAssets())
 		{
 			// Open a dialog asking the user to wait while assets are being discovered
 			SDiscoveringAssetsDialog::OpenDiscoveringAssetsDialog(
@@ -3195,11 +3214,48 @@ void UAssetToolsImpl::BeginAdvancedCopyPackages(const TArray<FName>& InputNamesT
 	}
 }
 
+TArray<FName> UAssetToolsImpl::ExpandAssetsAndFoldersToJustAssets(TArray<FName> SelectedAssetAndFolderNames) const
+{
+	IAssetRegistry& AssetRegistry = FAssetRegistryModule::GetRegistry();
+
+	TSet<FName> ExpandedAssets;
+	for (auto NameIt = SelectedAssetAndFolderNames.CreateIterator(); NameIt; ++NameIt)
+	{
+		FName OriginalName = *NameIt;
+		const FString& OriginalNameString = OriginalName.ToString();
+		if (!FPackageName::DoesPackageExist(OriginalNameString))
+		{
+			TArray<FAssetData> AssetsInFolder;
+			AssetRegistry.GetAssetsByPath(OriginalName, AssetsInFolder, true, false);
+			for (const FAssetData& Asset : AssetsInFolder)
+			{
+				ExpandedAssets.Add(Asset.PackageName);
+			}
+
+			NameIt.RemoveCurrent();
+		}
+		else
+		{
+			ExpandedAssets.Add(OriginalName);
+		}
+	}
+
+	TArray<FName> ExpandedAssetArray;
+	for (FName ExpandedAsset : ExpandedAssets)
+	{
+		ExpandedAssetArray.Add(ExpandedAsset);
+	}
+
+	return ExpandedAssetArray;
+}
+
 void UAssetToolsImpl::PerformAdvancedCopyPackages(TArray<FName> SelectedAssetAndFolderNames, FString TargetPath) const
 {	
 	TargetPath.RemoveFromEnd(TEXT("/"));
+
+	//TArray<FName> ExpandedAssets = ExpandAssetsAndFoldersToJustAssets(SelectedAssetAndFolderNames);
 	FAdvancedCopyParams CopyParams = FAdvancedCopyParams(SelectedAssetAndFolderNames, TargetPath);
-	CopyParams.bShouldCheckForDependencies = SelectedAssetAndFolderNames.Num() == 1;
+	CopyParams.bShouldCheckForDependencies = true;
 
 	// Suppress UI if we're running in unattended mode
 	if (FApp::IsUnattended())
@@ -3216,6 +3272,7 @@ void UAssetToolsImpl::PerformAdvancedCopyPackages(TArray<FName> SelectedAssetAnd
 		const FString& OriginalNameString = OriginalName.ToString();
 		FString SrcFilename;
 		UObject* ExistingObject = nullptr;
+
 		if (FPackageName::DoesPackageExist(OriginalNameString, nullptr, &SrcFilename))
 		{
 			UPackage* Pkg = LoadPackage(nullptr, *OriginalNameString, LOAD_None);
@@ -3225,6 +3282,7 @@ void UAssetToolsImpl::PerformAdvancedCopyPackages(TArray<FName> SelectedAssetAnd
 				ExistingObject = StaticFindObject(UObject::StaticClass(), Pkg, *Name);
 			}
 		}
+
 		if (ExistingObject)
 		{
 			// Try to find the customization in the settings
@@ -3260,7 +3318,7 @@ void UAssetToolsImpl::InitAdvancedCopyFromCopyParams(FAdvancedCopyParams CopyPar
 	TArray<TMap<FName, FName>> CompleteDependencyMap;
 	TArray<TMap<FString, FString>> CompleteDestinationMap;
 
-	TArray<FName> SelectedPackageNames = CopyParams.GetSelectedPackageNames();
+	TArray<FName> SelectedPackageNames = CopyParams.GetSelectedPackageOrFolderNames();
 
 	FScopedSlowTask SlowTask(SelectedPackageNames.Num(), LOCTEXT("AdvancedCopyPrepareSlowTask", "Preparing Files for Advanced Copy"));
 	SlowTask.MakeDialog();
@@ -3271,6 +3329,7 @@ void UAssetToolsImpl::InitAdvancedCopyFromCopyParams(FAdvancedCopyParams CopyPar
 	{
 		FName Package = SelectedPackageNames[CustomizationIndex];
 		SlowTask.EnterProgressFrame(1, FText::Format(LOCTEXT("AdvancedCopy_PreparingDependencies", "Preparing dependencies for {0}"), FText::FromString(Package.ToString())));
+
 		UAdvancedCopyCustomization* CopyCustomization = CustomizationsToUse[CustomizationIndex];
 		// Give the customization a chance to edit the copy parameters
 		CopyCustomization->EditCopyParams(CopyParams);
@@ -3279,13 +3338,12 @@ void UAssetToolsImpl::InitAdvancedCopyFromCopyParams(FAdvancedCopyParams CopyPar
 
 		// Get all packages to be copied
 		GetAllAdvancedCopySources(Package, CopyParams, PackageNamesToCopy, DependencyMap, CopyCustomization);
+		
 		// Allow the customization to apply any additional filters
 		CopyCustomization->ApplyAdditionalFiltering(PackageNamesToCopy);
-
 		CopyCustomization->SetPackageThatInitiatedCopy(Package.ToString());
-		TMap<FString, FString> DestinationMap = TMap<FString, FString>();
-
 
+		TMap<FString, FString> DestinationMap = TMap<FString, FString>();
 		GenerateAdvancedCopyDestinations(CopyParams, PackageNamesToCopy, CopyCustomization, DestinationMap);
 		CopyCustomization->TransformDestinationPaths(DestinationMap);
 		CompleteDestinationMap.Add(DestinationMap);
diff --git a/Engine/Source/Developer/AssetTools/Private/AssetTools.h b/Engine/Source/Developer/AssetTools/Private/AssetTools.h
--- a/Engine/Source/Developer/AssetTools/Private/AssetTools.h
+++ b/Engine/Source/Developer/AssetTools/Private/AssetTools.h
@@ -133,6 +133,8 @@ public:
 	/** The manager to handle fixing up redirectors */
 	TSharedPtr<FAssetFixUpRedirectors> AssetFixUpRedirectors;
 private:
+	static FString GenerateAdvancedCopyDestinationPackageName(const FString& SourcePackage, const FString& SourcePath, const FString& DestinationFolder);
+
 	/** Checks to see if a package is marked for delete then ask the user if he would like to check in the deleted file before he can continue. Returns true when it is safe to proceed. */
 	bool CheckForDeletedPackage(const UPackage* Package) const;
 
@@ -142,6 +144,8 @@ private:
 	/** Begins the package migration, after assets have been discovered */
 	void PerformMigratePackages(TArray<FName> PackageNamesToMigrate) const;
 
+	TArray<FName> ExpandAssetsAndFoldersToJustAssets(TArray<FName> SelectedAssetAndFolderNames) const;
+
 	/** Begins the package advanced copy, after assets have been discovered */
 	void PerformAdvancedCopyPackages(TArray<FName> SelectedPackageNames, FString TargetPath) const;
 
diff --git a/Engine/Source/Developer/AssetTools/Private/SAdvancedCopyReportDialog.cpp b/Engine/Source/Developer/AssetTools/Private/SAdvancedCopyReportDialog.cpp
--- a/Engine/Source/Developer/AssetTools/Private/SAdvancedCopyReportDialog.cpp
+++ b/Engine/Source/Developer/AssetTools/Private/SAdvancedCopyReportDialog.cpp
@@ -19,6 +19,7 @@
 #include "AssetToolsSettings.h"
 #include "Misc/PackageName.h"
 #include "Widgets/Input/SEditableTextBox.h"
+#include "AssetRegistry/AssetRegistryModule.h"
 
 #define LOCTEXT_NAMESPACE "AdvancedCopyReportDialog"
 
@@ -274,9 +275,24 @@ void SAdvancedCopyReportDialog::Construct( const FArguments& InArgs, const FAdva
 	];
 
 	// Make sure the initially selected packages begin as part of the set we're definitely cloning.
-	for (FName SelectedPackageName : InParams.GetSelectedPackageNames())
+	IAssetRegistry& AssetRegistry = FAssetRegistryModule::GetRegistry();
+
+	for (FName OriginalName : InParams.GetSelectedPackageOrFolderNames())
 	{
-		CloneSet->Add(SelectedPackageName.ToString());
+		const FString& OriginalNameString = OriginalName.ToString();
+		if (!FPackageName::DoesPackageExist(OriginalNameString))
+		{
+			TArray<FAssetData> AssetsInFolder;
+			AssetRegistry.GetAssetsByPath(OriginalName, AssetsInFolder, true, false);
+			for (const FAssetData& Asset : AssetsInFolder)
+			{
+				CloneSet->Add(Asset.PackageName.ToString());
+			}
+		}
+		else
+		{
+			CloneSet->Add(OriginalName.ToString());
+		}
 	}
 }
 END_SLATE_FUNCTION_BUILD_OPTIMIZATION
@@ -398,9 +414,14 @@ FReply SAdvancedCopyReportDialog::OkClicked()
 					FString OutPackageRoot, OutPackagePath, OutPackageName;
 					FPackageName::SplitLongPackageName(It.Value(), OutPackageRoot, OutPackagePath, OutPackageName);
 
-					FString RenamedPackageName = OutPackageName.Replace(*FindString, *ReplaceString);
+					FString CroppedDestination = It.Value();
+					const bool bDidRemovePrefix = CroppedDestination.RemoveFromStart(ReportString);
+					if (bDidRemovePrefix)
+					{
+						CroppedDestination = CroppedDestination;
 
-					It.Value() = OutPackageRoot / OutPackagePath / RenamedPackageName;
+						It.Value() = ReportString / CroppedDestination.Replace(*FindString, *ReplaceString);
+					}
 				}
 			}
 		}
@@ -540,6 +561,7 @@ const TSharedRef< SWidget > SAdvancedCopyColumn::ConstructRowWidget(TSharedPtr<s
 		if (bDidRemovePrefix)
 		{
 			CroppedDestination = TEXT(".") + CroppedDestination;
+
 		}
 		return SNew(STextBlock)
 			.Text(FText::FromString(FPaths::GetPath(CroppedDestination)));
diff --git a/Engine/Source/Developer/AssetTools/Public/IAssetTools.h b/Engine/Source/Developer/AssetTools/Public/IAssetTools.h
--- a/Engine/Source/Developer/AssetTools/Public/IAssetTools.h
+++ b/Engine/Source/Developer/AssetTools/Public/IAssetTools.h
@@ -107,15 +107,20 @@ struct FAdvancedCopyParams
 {
 	GENERATED_USTRUCT_BODY()
 
-	bool bShouldForceSave;
-	bool bCopyOverAllDestinationOverlaps;
-	bool bGenerateUniqueNames;
-	bool bShouldSuppressUI;
-	bool bShouldCheckForDependencies;
+	bool bShouldForceSave = false;
+	bool bCopyOverAllDestinationOverlaps = false;
+	bool bShouldSuppressUI = false;
+	bool bShouldCheckForDependencies = false;
 
+	UE_DEPRECATED(5.0, "This function has been deprecated, use GetSelectedPackageOrFolderNames")
 	const TArray<FName>& GetSelectedPackageNames() const
 	{
-		return SelectedPackageNames;
+		return SelectedPackageOrFolderNames;
+	}
+
+	const TArray<FName>& GetSelectedPackageOrFolderNames() const
+	{
+		return SelectedPackageOrFolderNames;
 	}
 
 	const FString& GetDropLocationForAdvancedCopy() const
@@ -133,13 +138,12 @@ struct FAdvancedCopyParams
 		CustomizationsToUse.Add(InCustomization);
 	}
 
-	FAdvancedCopyParams(TArray<FName> InSelectedPackageNames, FString InDropLocationForAdvancedCopy)
+	FAdvancedCopyParams(TArray<FName> InSelectedPackageOrFolderNames, FString InDropLocationForAdvancedCopy)
 		: bShouldForceSave(false)
 		, bCopyOverAllDestinationOverlaps(true)
-		, bGenerateUniqueNames(false)
 		, bShouldSuppressUI(false)
 		, bShouldCheckForDependencies(true)
-		, SelectedPackageNames(InSelectedPackageNames)
+		, SelectedPackageOrFolderNames(InSelectedPackageOrFolderNames)
 		, DropLocationForAdvancedCopy(InDropLocationForAdvancedCopy)
 	{
 	}
@@ -147,10 +151,9 @@ struct FAdvancedCopyParams
 	FAdvancedCopyParams() {}
 
 private:
-	TArray<FName> SelectedPackageNames;
+	TArray<FName> SelectedPackageOrFolderNames;
 	TArray<UAdvancedCopyCustomization*> CustomizationsToUse;
 	FString DropLocationForAdvancedCopy;
-
 };
 
 UINTERFACE(MinimalApi, BlueprintType, meta = (CannotImplementInterfaceInBlueprint))
