Fixing a build error.

diff --git a/Engine/Source/Developer/AssetTools/Private/SAdvancedCopyReportDialog.cpp b/Engine/Source/Developer/AssetTools/Private/SAdvancedCopyReportDialog.cpp
--- a/Engine/Source/Developer/AssetTools/Private/SAdvancedCopyReportDialog.cpp
+++ b/Engine/Source/Developer/AssetTools/Private/SAdvancedCopyReportDialog.cpp
@@ -418,8 +418,6 @@ FReply SAdvancedCopyReportDialog::OkClicked()
 					const bool bDidRemovePrefix = CroppedDestination.RemoveFromStart(ReportString);
 					if (bDidRemovePrefix)
 					{
-						CroppedDestination = CroppedDestination;
-
 						It.Value() = ReportString / CroppedDestination.Replace(*FindString, *ReplaceString);
 					}
 				}
