Minimal editor changes from a partner branch.

diff --git a/Engine/Plugins/Experimental/AlembicImporter/Source/AlembicImporter/Classes/AlembicImportFactory.h b/Engine/Plugins/Experimental/AlembicImporter/Source/AlembicImporter/Classes/AlembicImportFactory.h
--- a/Engine/Plugins/Experimental/AlembicImporter/Source/AlembicImporter/Classes/AlembicImportFactory.h
+++ b/Engine/Plugins/Experimental/AlembicImporter/Source/AlembicImporter/Classes/AlembicImportFactory.h
@@ -17,8 +17,8 @@ class SAlembicTrackSelectionWindow;
 
 class FAbcImporter;
 
-UCLASS(hidecategories=Object)
-class UAlembicImportFactory : public UFactory, public FReimportHandler
+UCLASS(hidecategories = Object)
+class ALEMBICIMPORTER_API UAlembicImportFactory : public UFactory, public FReimportHandler
 {
 	GENERATED_UCLASS_BODY()
 
diff --git a/Engine/Source/Developer/AssetTools/Private/AssetTools.cpp b/Engine/Source/Developer/AssetTools/Private/AssetTools.cpp
--- a/Engine/Source/Developer/AssetTools/Private/AssetTools.cpp
+++ b/Engine/Source/Developer/AssetTools/Private/AssetTools.cpp
@@ -728,7 +728,8 @@ void UAssetToolsImpl::GenerateAdvancedCopyDestinations(FAdvancedCopyParams& InPa
 				FString Parent = FString();
 				if (bGenerateRelativePaths)
 				{
-					FString RootFolder = UAdvancedCopyCustomization::StaticClass()->GetDefaultObject<UAdvancedCopyCustomization>()->GetPackageThatInitiatedCopy();
+
+					FString RootFolder = CopyCustomization->GetPackageThatInitiatedCopy();
 					if (RootFolder != PackageNameString)
 					{
 						FString BaseParent = FString();
@@ -3199,7 +3200,8 @@ void UAssetToolsImpl::RecursiveGetDependenciesAdvanced(const FName& PackageName,
 		// We found a folder containing assets
 		if (AssetRegistry.HasAssets(PackageName) && AssetRegistry.GetAssetsByPath(PackageName, PathAssetData))
 		{
-			FARFilter ExclusionFilter = UAdvancedCopyCustomization::StaticClass()->GetDefaultObject<UAdvancedCopyCustomization>()->GetARFilter();
+
+			FARFilter ExclusionFilter = CopyCustomization->GetARFilter();
 			AssetRegistry.UseFilterToExcludeAssets(PathAssetData, ExclusionFilter);
 			for(const FAssetData& Asset : PathAssetData)
 			{
