Updated Nanite options for Static mesh context menu in Content Browser

diff --git a/Engine/Source/Developer/AssetTools/Private/AssetTypeActions/AssetTypeActions_StaticMesh.cpp b/Engine/Source/Developer/AssetTools/Private/AssetTypeActions/AssetTypeActions_StaticMesh.cpp
--- a/Engine/Source/Developer/AssetTools/Private/AssetTypeActions/AssetTypeActions_StaticMesh.cpp
+++ b/Engine/Source/Developer/AssetTools/Private/AssetTypeActions/AssetTypeActions_StaticMesh.cpp
@@ -114,18 +114,77 @@ void FAssetTypeActions_StaticMesh::GetResolvedSourceFilePaths(const TArray<UObje
 
 void FAssetTypeActions_StaticMesh::GetNaniteMenu(class FMenuBuilder& MenuBuilder, TArray<TWeakObjectPtr<UStaticMesh>> Meshes)
 {
+	if (Meshes.Num() == 1)
+	{
+		TWeakObjectPtr<UStaticMesh> First = Meshes[0];
+		UStaticMesh* StaticMesh = First.Get();
+
+		MenuBuilder.AddMenuEntry(
+			NSLOCTEXT("AssetTypeActions_StaticMesh", "StaticMesh_NaniteToggle", "Nanite"),
+			NSLOCTEXT("AssetTypeActions_StaticMesh", "StaticMesh_NaniteToggleTooltip", "Toggle Nanite support on the selected mesh."),
+			FSlateIcon(),
+			FUIAction(
+				FExecuteAction::CreateLambda([this, Meshes, StaticMesh]()
+					{
+						if (StaticMesh->NaniteSettings.bEnabled)
+						{
+							ExecuteNaniteDisable(Meshes);
+						}
+						else
+						{
+							ExecuteNaniteEnable(Meshes);
+						}
+					}),
+				FCanExecuteAction(),
+				FIsActionChecked::CreateLambda([StaticMesh]()
+					{
+						return static_cast<bool>(StaticMesh->NaniteSettings.bEnabled);
+					})
+			),
+			NAME_None,
+			EUserInterfaceActionType::ToggleButton
+		);
+	}
+	else
+	{
+		// Dummy action that cannot be executed to indicate when multiple meshes are selected
+		MenuBuilder.AddMenuEntry(
+			NSLOCTEXT("AssetTypeActions_StaticMesh", "StaticMesh_NaniteMultipleSelected", "Multiple Meshes Selected"),
+			TAttribute<FText>(),
+			FSlateIcon(),
+			FUIAction(
+				FExecuteAction(),
+				FCanExecuteAction::CreateLambda([]() { return false; })
+			)
+		);
+	}
+
+	MenuBuilder.AddMenuSeparator();
+
 	MenuBuilder.AddMenuEntry(
-		NSLOCTEXT("AssetTypeActions_StaticMesh", "StaticMesh_NaniteEnable", "Enable"),
-		NSLOCTEXT("AssetTypeActions_StaticMesh", "StaticMesh_NaniteEnableTooltip", "Enables support for Nanite on the selected mesh(es)."),
+		NSLOCTEXT("AssetTypeActions_StaticMesh", "StaticMesh_NaniteEnableAll", "Enable Nanite for Selected"),
+		NSLOCTEXT("AssetTypeActions_StaticMesh", "StaticMesh_NaniteEnableAllTooltip", "Enables support for Nanite on the selected meshes."),
 		FSlateIcon(),
-		FUIAction(FExecuteAction::CreateSP(this, &FAssetTypeActions_StaticMesh::ExecuteNaniteEnable, Meshes))
+		FUIAction(
+			FExecuteAction::CreateSP(this, &FAssetTypeActions_StaticMesh::ExecuteNaniteEnable, Meshes),
+			FCanExecuteAction::CreateLambda([Meshes]()
+				{
+					return Meshes.Num() > 1;
+				})
+		)
 	);
 
 	MenuBuilder.AddMenuEntry(
-		NSLOCTEXT("AssetTypeActions_StaticMesh", "StaticMesh_NaniteDisable", "Disable"),
-		NSLOCTEXT("AssetTypeActions_StaticMesh", "StaticMesh_NaniteDisableTooltip", "Disables support for Nanite on the selected mesh(es)."),
+		NSLOCTEXT("AssetTypeActions_StaticMesh", "StaticMesh_NaniteDisableAll", "Disable Nanite for Selected"),
+		NSLOCTEXT("AssetTypeActions_StaticMesh", "StaticMesh_NaniteDisableAllTooltip", "Disables support for Nanite on the selected meshes."),
 		FSlateIcon(),
-		FUIAction(FExecuteAction::CreateSP(this, &FAssetTypeActions_StaticMesh::ExecuteNaniteDisable, Meshes))
+		FUIAction(
+			FExecuteAction::CreateSP(this, &FAssetTypeActions_StaticMesh::ExecuteNaniteDisable, Meshes),
+			FCanExecuteAction::CreateLambda([Meshes]()
+				{
+					return Meshes.Num() > 1;
+				})
+			)
 	);
 }
 
