Add option to create child anim BP with another (compatible) skeleton

Works with skeleton remapping and template anim BPs

diff --git a/Engine/Source/Developer/AssetTools/Private/AssetTypeActions/AssetTypeActions_AnimBlueprint.cpp b/Engine/Source/Developer/AssetTools/Private/AssetTypeActions/AssetTypeActions_AnimBlueprint.cpp
--- a/Engine/Source/Developer/AssetTools/Private/AssetTypeActions/AssetTypeActions_AnimBlueprint.cpp
+++ b/Engine/Source/Developer/AssetTools/Private/AssetTypeActions/AssetTypeActions_AnimBlueprint.cpp
@@ -10,6 +10,7 @@
 #include "Factories/AnimBlueprintFactory.h"
 #include "ThumbnailRendering/SceneThumbnailInfo.h"
 #include "AssetTools.h"
+#include "ContentBrowserModule.h"
 #include "PersonaModule.h"
 #include "Framework/Notifications/NotificationManager.h"
 #include "SBlueprintDiff.h"
@@ -17,6 +18,8 @@
 #include "SSkeletonWidget.h"
 #include "Styling/SlateIconFinder.h"
 #include "IAnimationBlueprintEditorModule.h"
+#include "IContentBrowserSingleton.h"
+#include "Kismet2/KismetEditorUtilities.h"
 #include "Preferences/PersonaOptions.h"
 #if WITH_EDITOR
 #include "Subsystems/AssetEditorSubsystem.h"
@@ -29,8 +32,102 @@ void FAssetTypeActions_AnimBlueprint::GetActions(const TArray<UObject*>& InObjec
 {
 	FAssetTypeActions_Blueprint::GetActions(InObjects, Section);
 
-	auto AnimBlueprints = GetTypedWeakObjectPtrs<UAnimBlueprint>(InObjects);
+	TArray<TWeakObjectPtr<UAnimBlueprint>> AnimBlueprints = GetTypedWeakObjectPtrs<UAnimBlueprint>(InObjects);
 
+	if (AnimBlueprints.Num() == 1 && CanCreateNewDerivedBlueprint())
+	{
+		UAnimBlueprint* AnimBlueprint = AnimBlueprints[0].Get();
+
+		// Accept (non-interface) template anim BPs or anim BPs with compatible skeletons
+		if(AnimBlueprint && AnimBlueprint->BlueprintType != BPTYPE_Interface && ((AnimBlueprint->TargetSkeleton == nullptr && AnimBlueprint->bIsTemplate) || AnimBlueprint->TargetSkeleton->GetCompatibleSkeletons().Num() > 0))
+		{
+			Section.AddSubMenu(
+				"AnimBlueprint_NewSkeletonChildBlueprint",
+				LOCTEXT("AnimBlueprint_NewSkeletonChildBlueprint", "Create Child Anim Blueprint with Skeleton"),
+				LOCTEXT("AnimBlueprint_NewSkeletonChildBlueprint_Tooltip", "Create a child Anim Blueprint that uses a different compatible skeleton"),
+				FNewToolMenuDelegate::CreateLambda([this, WeakAnimBlueprint = TWeakObjectPtr<UAnimBlueprint>(AnimBlueprint)](UToolMenu* CompatibleSkeletonMenu)
+				{
+					auto HandleAssetSelected = [this, WeakAnimBlueprint](const FAssetData& InAssetData)
+					{
+						FSlateApplication::Get().DismissAllMenus();
+						
+						if(UAnimBlueprint* TargetParentBP = WeakAnimBlueprint.Get())
+						{
+							USkeleton* TargetSkeleton = CastChecked<USkeleton>(InAssetData.GetAsset());
+							UClass* TargetParentClass = TargetParentBP->GeneratedClass;
+
+							if (!FKismetEditorUtilities::CanCreateBlueprintOfClass(TargetParentClass))
+							{
+								FMessageDialog::Open( EAppMsgType::Ok, LOCTEXT("InvalidClassToMakeBlueprintFrom", "Invalid class with which to make a Blueprint."));
+								return;
+							}
+
+							FString Name;
+							FString PackageName;
+							CreateUniqueAssetName(TargetParentBP->GetOutermost()->GetName(), TEXT("_Child"), PackageName, Name);
+							const FString PackagePath = FPackageName::GetLongPackagePath(PackageName);
+
+							UAnimBlueprintFactory* AnimBlueprintFactory = NewObject<UAnimBlueprintFactory>();
+							AnimBlueprintFactory->ParentClass = TSubclassOf<UAnimInstance>(*TargetParentBP->GeneratedClass);
+							AnimBlueprintFactory->TargetSkeleton = TargetSkeleton;
+
+							FContentBrowserModule& ContentBrowserModule = FModuleManager::LoadModuleChecked<FContentBrowserModule>("ContentBrowser");
+							ContentBrowserModule.Get().CreateNewAsset(Name, PackagePath, TargetParentBP->GetClass(), AnimBlueprintFactory);
+						}
+					};
+
+					FAssetPickerConfig AssetPickerConfig;
+					AssetPickerConfig.OnAssetEnterPressed = FOnAssetEnterPressed::CreateLambda([HandleAssetSelected](const TArray<FAssetData>& SelectedAssetData)
+					{
+						if (SelectedAssetData.Num() == 1)
+						{
+							HandleAssetSelected(SelectedAssetData[0]);
+						}
+					});
+					AssetPickerConfig.OnAssetSelected = FOnAssetSelected::CreateLambda(HandleAssetSelected);
+					AssetPickerConfig.bAllowNullSelection = false;
+					AssetPickerConfig.InitialAssetViewType = EAssetViewType::List;
+					AssetPickerConfig.Filter.bRecursiveClasses = false;
+					AssetPickerConfig.Filter.ClassNames.Add(USkeleton::StaticClass()->GetFName());
+					AssetPickerConfig.OnShouldFilterAsset = FOnShouldFilterAsset::CreateLambda([WeakAnimBlueprint](const FAssetData& AssetData)
+					{
+						if(UAnimBlueprint* LocalAnimBlueprint = WeakAnimBlueprint.Get())
+						{
+							if(LocalAnimBlueprint->TargetSkeleton == nullptr && LocalAnimBlueprint->bIsTemplate)
+							{
+								// Template anim BP - do not filter
+								return false;
+							}
+							else if(LocalAnimBlueprint->TargetSkeleton != nullptr)
+							{
+								// Filter on compatible skeletons
+								const FString ExportTextName = AssetData.GetExportTextName();
+								return !LocalAnimBlueprint->TargetSkeleton->IsCompatibleSkeletonByAssetString(ExportTextName);
+							}
+						}
+						
+						return true;
+					});
+					
+					FContentBrowserModule& ContentBrowserModule = FModuleManager::Get().LoadModuleChecked<FContentBrowserModule>(TEXT("ContentBrowser"));
+
+					FToolMenuSection& InSection = CompatibleSkeletonMenu->AddSection("CompatibleSkeletonMenu", LOCTEXT("CompatibleSkeletonHeader", "Compatible Skeletons"));
+					InSection.AddEntry(
+						FToolMenuEntry::InitWidget("CompatibleSkeletonPicker",
+							SNew(SBox)
+							.WidthOverride(300.0f)
+							.HeightOverride(300.0f)
+							[
+								ContentBrowserModule.Get().CreateAssetPicker(AssetPickerConfig)
+							],
+							FText::GetEmpty())
+					);
+				}),
+				false,
+				FSlateIcon(FAppStyle::GetAppStyleSetName(), "Icons.Blueprint"));
+		}
+	}
+	
 	Section.AddMenuEntry(
 		"AnimBlueprint_FindSkeleton",
 		LOCTEXT("AnimBlueprint_FindSkeleton", "Find Skeleton"),
