Stop interchange from being use when using asset tool api by defaul (when interchange is on). (Otherwise this would break some licensees made code and scripts). Those who want an async import should use the interchange api preferably.

Added a setting to turn on interchange import only for textures.

The selection of the asset in the content is now only done after all the asset are imported when using interchange.

diff --git a/Engine/Plugins/Experimental/Interchange/Runtime/Source/Import/Private/InterchangeImportModule.cpp b/Engine/Plugins/Experimental/Interchange/Runtime/Source/Import/Private/InterchangeImportModule.cpp
--- a/Engine/Plugins/Experimental/Interchange/Runtime/Source/Import/Private/InterchangeImportModule.cpp
+++ b/Engine/Plugins/Experimental/Interchange/Runtime/Source/Import/Private/InterchangeImportModule.cpp
@@ -55,19 +55,31 @@ void FInterchangeImportModule::StartupModule()
 
 		// UDIM must be registered before the other texture translators
 		InterchangeManager.RegisterTranslator(UInterchangeUDIMTranslator::StaticClass());
+		InterchangeManager.RegisterTextureOnlyTranslatorClass(UInterchangeUDIMTranslator::StaticClass());
 
 		InterchangeManager.RegisterTranslator(UInterchangeBMPTranslator::StaticClass());
+		InterchangeManager.RegisterTextureOnlyTranslatorClass(UInterchangeBMPTranslator::StaticClass());
 		InterchangeManager.RegisterTranslator(UInterchangeDDSTranslator::StaticClass());
+		InterchangeManager.RegisterTextureOnlyTranslatorClass(UInterchangeDDSTranslator::StaticClass());
 		InterchangeManager.RegisterTranslator(UInterchangeEXRTranslator::StaticClass());
+		InterchangeManager.RegisterTextureOnlyTranslatorClass(UInterchangeEXRTranslator::StaticClass());
 		InterchangeManager.RegisterTranslator(UInterchangeJPGTranslator::StaticClass());
+		InterchangeManager.RegisterTextureOnlyTranslatorClass(UInterchangeJPGTranslator::StaticClass());
 		InterchangeManager.RegisterTranslator(UInterchangePCXTranslator::StaticClass());
+		InterchangeManager.RegisterTextureOnlyTranslatorClass(UInterchangePCXTranslator::StaticClass());
 		InterchangeManager.RegisterTranslator(UInterchangePNGTranslator::StaticClass());
+		InterchangeManager.RegisterTextureOnlyTranslatorClass(UInterchangePNGTranslator::StaticClass());
 		InterchangeManager.RegisterTranslator(UInterchangePSDTranslator::StaticClass());
+		InterchangeManager.RegisterTextureOnlyTranslatorClass(UInterchangePSDTranslator::StaticClass());
 		InterchangeManager.RegisterTranslator(UInterchangeTGATranslator::StaticClass());
+		InterchangeManager.RegisterTextureOnlyTranslatorClass(UInterchangeTGATranslator::StaticClass());
 		InterchangeManager.RegisterTranslator(UInterchangeHDRTranslator::StaticClass());
+		InterchangeManager.RegisterTextureOnlyTranslatorClass(UInterchangeHDRTranslator::StaticClass());
 		InterchangeManager.RegisterTranslator(UInterchangeIESTranslator::StaticClass());
+		InterchangeManager.RegisterTextureOnlyTranslatorClass(UInterchangeIESTranslator::StaticClass());
 #if WITH_LIBTIFF
 		InterchangeManager.RegisterTranslator(UInterchangeTIFFTranslator::StaticClass());
+		InterchangeManager.RegisterTextureOnlyTranslatorClass(UInterchangeTIFFTranslator::StaticClass());
 #endif // WITH_LIBTIFF
 
 		//Register the factories
diff --git a/Engine/Source/Developer/AssetTools/Private/AssetTools.cpp b/Engine/Source/Developer/AssetTools/Private/AssetTools.cpp
--- a/Engine/Source/Developer/AssetTools/Private/AssetTools.cpp
+++ b/Engine/Source/Developer/AssetTools/Private/AssetTools.cpp
@@ -1328,7 +1328,7 @@ void UAssetToolsImpl::ExpandDirectories(const TArray<FString>& Files, const FStr
 		}
 	}
 }
-TArray<UObject*> UAssetToolsImpl::ImportAssets(const TArray<FString>& Files, const FString& DestinationPath, UFactory* ChosenFactory, bool bSyncToBrowser /* = true */, TArray<TPair<FString, FString>>* FilesAndDestinations /* = nullptr */) const
+TArray<UObject*> UAssetToolsImpl::ImportAssets(const TArray<FString>& Files, const FString& DestinationPath, UFactory* ChosenFactory, bool bSyncToBrowser /* = true */, TArray<TPair<FString, FString>>* FilesAndDestinations /* = nullptr */, bool bAllowAsyncImport /* = false */) const
 {
 	const bool bForceOverrideExisting = false;
 
@@ -1338,6 +1338,7 @@ TArray<UObject*> UAssetToolsImpl::ImportAssets(const TArray<FString>& Files, con
 	Params.bForceOverrideExisting = false;
 	Params.bSyncToBrowser = bSyncToBrowser;
 	Params.SpecifiedFactory = TStrongObjectPtr<UFactory>(ChosenFactory);
+	Params.bAllowAsyncImport = bAllowAsyncImport;
 
 	return ImportAssetsInternal(Files, DestinationPath, FilesAndDestinations, Params);
 }
@@ -1746,9 +1747,11 @@ TArray<UObject*> UAssetToolsImpl::ImportAssetsInternal(const TArray<FString>& Fi
 	FScopedSlowTask SlowTask(ValidFiles.Num(), LOCTEXT("ImportSlowTask", "Importing"));
 
 	bool bUseInterchangeFramework = false;
+	bool bUseInterchangeFrameworkForTextureOnly = false;
 	UInterchangeManager& InterchangeManager = UInterchangeManager::GetInterchangeManager();
 #if WITH_EDITOR
-	bUseInterchangeFramework = GetDefault<UEditorExperimentalSettings>()->bEnableInterchangeFramework;
+	const UEditorExperimentalSettings* EditorExperimentalSettings = GetDefault<UEditorExperimentalSettings>();
+	bUseInterchangeFramework = EditorExperimentalSettings->bEnableInterchangeFramework;
 
 	if (bUseInterchangeFramework && Params.SpecifiedFactory)
 	{
@@ -1757,8 +1760,14 @@ TArray<UObject*> UAssetToolsImpl::ImportAssetsInternal(const TArray<FString>& Fi
 			bUseInterchangeFramework = GetDefault<UInterchangeProjectSettings>()->bUseInterchangeWhenImportingIntoLevel;
 		}
 	}
+
+	bUseInterchangeFrameworkForTextureOnly = (!bUseInterchangeFramework) && EditorExperimentalSettings->bEnableInterchangeFrameworkForTextureOnly;
+	bUseInterchangeFramework |= bUseInterchangeFrameworkForTextureOnly;
 #endif
 
+	// Block interchange use if the user is not aware that is import can be async. Otherwise, we don't return the imported object and we can't mimic the SpecifiedFactory settings.
+	bUseInterchangeFramework &= Params.bAllowAsyncImport;
+
 	if (!bUseInterchangeFramework && ValidFiles.Num() > 1)
 	{	
 		//Always allow user to cancel the import task if they are importing multiple ValidFiles.
@@ -1766,7 +1775,6 @@ TArray<UObject*> UAssetToolsImpl::ImportAssetsInternal(const TArray<FString>& Fi
 		SlowTask.MakeDialog(true);
 	}
 
-
 	TArray<TPair<FString, FString>> FilesAndDestinations;
 	if (FilesAndDestinationsPtr == nullptr)
 	{
@@ -1910,7 +1918,17 @@ TArray<UObject*> UAssetToolsImpl::ImportAssetsInternal(const TArray<FString>& Fi
 			const FString& Filename = FilesAndDestinations[FileIdx].Key;
 			{
 				UE::Interchange::FScopedSourceData ScopedSourceData(Filename);
-				if (!InterchangeManager.CanTranslateSourceData(ScopedSourceData.GetSourceData()))
+
+				if (bUseInterchangeFrameworkForTextureOnly)
+				{
+					UInterchangeTranslatorBase* Translator = InterchangeManager.GetTranslatorForSourceData(ScopedSourceData.GetSourceData());
+					if (!Translator || !InterchangeManager.IsTranslatorClassForTextureOnly(Translator->GetClass()))
+					{
+						bOnlyInterchangeImport = false;
+						break;
+					}
+				}
+				else if (!InterchangeManager.CanTranslateSourceData(ScopedSourceData.GetSourceData()))
 				{
 					bOnlyInterchangeImport = false;
 					break;
@@ -1942,6 +1960,7 @@ TArray<UObject*> UAssetToolsImpl::ImportAssetsInternal(const TArray<FString>& Fi
 		{}
 
 		TStrongObjectPtr<UInterchangeResultsContainer> InterchangeResultsContainer;
+		TArray<TWeakObjectPtr<UObject>> ImportedObjects;
 		std::atomic<int32> ImportCount;
 	};
 
@@ -1974,39 +1993,57 @@ TArray<UObject*> UAssetToolsImpl::ImportAssetsInternal(const TArray<FString>& Fi
 		if (bUseInterchangeFramework)
 		{
 			UE::Interchange::FScopedSourceData ScopedSourceData(Filename);
-			if (InterchangeManager.CanTranslateSourceData(ScopedSourceData.GetSourceData()))
+
+			bool bUseATextureTranslator = false;
+			if (bUseInterchangeFrameworkForTextureOnly)
 			{
-				auto PostImportedLambda = [bSyncToBrowser](UObject* ImportedObject)
+				UInterchangeTranslatorBase* Translator = InterchangeManager.GetTranslatorForSourceData(ScopedSourceData.GetSourceData());
+				if (Translator && InterchangeManager.IsTranslatorClassForTextureOnly(Translator->GetClass()))
 				{
-					if (ImportedObject && (bSyncToBrowser != false))
-					{
-						TArray<UObject*> ObjectArray;
-						ObjectArray.Add(ImportedObject);
-						UAssetToolsImpl::Get().SyncBrowserToAssets(ObjectArray);
-					}
-				};
-				FDelegateHandle PostImportHandle = InterchangeManager.OnAssetPostImport.AddLambda(PostImportedLambda);
-			
+					bUseATextureTranslator = true;
+				}
+			}
+
+			if (bUseATextureTranslator || (!bUseInterchangeFrameworkForTextureOnly && InterchangeManager.CanTranslateSourceData(ScopedSourceData.GetSourceData())))
+			{
 				FImportAssetParameters ImportAssetParameters;
 				ImportAssetParameters.bIsAutomated = bAutomatedImport;
 				ImportAssetParameters.ReimportAsset = nullptr;
 
+
 				TFunction<void(UE::Interchange::FImportResult&)> AppendImportResult =
 					// Note: ImportStatus captured by value so that the lambda keeps the shared ptr alive
 					[ImportStatus](UE::Interchange::FImportResult& Result)
 					{
 						ImportStatus->InterchangeResultsContainer->Append(Result.GetResults());
+						ImportStatus->ImportedObjects.Append(Result.GetImportedObjects());
+
 					};
 
 				TFunction<void(UE::Interchange::FImportResult&)> AppendAndBroadcastImportResultIfNeeded =
 					// Note: ImportStatus captured by value so that the lambda keeps the shared ptr alive
-					[ImportStatus, AppendImportResult](UE::Interchange::FImportResult& Result)
+					[ImportStatus, AppendImportResult, bSyncToBrowser](UE::Interchange::FImportResult& Result)
 					{
 						AppendImportResult(Result);
+
 						if (--ImportStatus->ImportCount == 0)
 						{
 							UInterchangeManager& InterchangeManager = UInterchangeManager::GetInterchangeManager();
 							InterchangeManager.OnBatchImportComplete.Broadcast(ImportStatus->InterchangeResultsContainer);
+
+							if (bSyncToBrowser)
+							{
+								// Only sync the content browser when the full import is done. Otherwise it can be annoying for the user.
+								// UX suggestion : Maybe we could move this into the post import window as button ("Select Imported Asset(s)") so it would be even less disruptive for the users.
+								TArray<UObject*> ImportedObjects;
+								ImportedObjects.Reserve(ImportStatus->ImportedObjects.Num());
+								for (const TWeakObjectPtr<UObject>& WeakObject : ImportStatus->ImportedObjects)
+								{
+									ImportedObjects.Add(WeakObject.Get());
+								}
+
+								UAssetToolsImpl::Get().SyncBrowserToAssets(ImportedObjects);
+							}
 						}
 					};
 
@@ -2023,8 +2060,7 @@ TArray<UObject*> UAssetToolsImpl::ImportAssetsInternal(const TArray<FString>& Fi
 					UE::Interchange::FAssetImportResultRef InterchangeResult = (InterchangeManager.ImportAssetAsync(DestinationPath, ScopedSourceData.GetSourceData(), ImportAssetParameters));
 					InterchangeResult->OnDone(AppendAndBroadcastImportResultIfNeeded);
 				}
-				
-				InterchangeManager.OnAssetPostImport.Remove(PostImportHandle);
+
 				//Import done, iterate the next file and destination
 				
 				//If we do not import only interchange file, update the progress for each interchange task
diff --git a/Engine/Source/Developer/AssetTools/Private/AssetTools.h b/Engine/Source/Developer/AssetTools/Private/AssetTools.h
--- a/Engine/Source/Developer/AssetTools/Private/AssetTools.h
+++ b/Engine/Source/Developer/AssetTools/Private/AssetTools.h
@@ -29,6 +29,7 @@ struct FAssetImportParams
 		, bSyncToBrowser(true)
 		, bForceOverrideExisting(false)
 		, bAutomated(false)
+		, bAllowAsyncImport(false)
 	{}
 
 	/** Factory to use for importing files */
@@ -43,6 +44,8 @@ struct FAssetImportParams
 	bool bForceOverrideExisting : 1;
 	/** Whether or not this is an automated import */
 	bool bAutomated : 1;
+	/** Temporary setting to enable the async import. (We will figure a better solution in a later release) */
+	bool bAllowAsyncImport : 1;
 };
 
 
@@ -85,7 +88,7 @@ public:
 	virtual void RenameReferencingSoftObjectPaths(const TArray<UPackage *> PackagesToCheck, const TMap<FSoftObjectPath, FSoftObjectPath>& AssetRedirectorMap) override;
 	virtual TArray<UObject*> ImportAssets(const FString& DestinationPath) override;
 	virtual TArray<UObject*> ImportAssetsWithDialog(const FString& DestinationPath) override;
-	virtual TArray<UObject*> ImportAssets(const TArray<FString>& Files, const FString& DestinationPath, UFactory* ChosenFactory, bool bSyncToBrowser = true, TArray<TPair<FString, FString>>* FilesAndDestinations = nullptr) const override;
+	virtual TArray<UObject*> ImportAssets(const TArray<FString>& Files, const FString& DestinationPath, UFactory* ChosenFactory, bool bSyncToBrowser = true, TArray<TPair<FString, FString>>* FilesAndDestinations = nullptr, bool bAllowAsyncImport = false) const override;
 	virtual TArray<UObject*> ImportAssetsAutomated(const UAutomatedAssetImportData* ImportData) override;
 	virtual void ImportAssetTasks(const TArray<UAssetImportTask*>& ImportTasks) override;
 	virtual void ExportAssets(const TArray<FString>& AssetsToExport, const FString& ExportPath) override;
diff --git a/Engine/Source/Developer/AssetTools/Public/IAssetTools.h b/Engine/Source/Developer/AssetTools/Public/IAssetTools.h
--- a/Engine/Source/Developer/AssetTools/Public/IAssetTools.h
+++ b/Engine/Source/Developer/AssetTools/Public/IAssetTools.h
@@ -293,9 +293,10 @@ public:
 	 * @param DestinationPath	destination path for imported files
 	 * @param ChosenFactory		Specific factory to use for object creation
 	 * @param bSyncToBrowser	If true sync content browser to first imported asset after import
+	 * @param bAllowAsyncImport	This allow the import code to use a async importer if enabled. (Note doing so will ignore the ChosenFactory arguments (If you want a async import prefer the InterchangeManager api))
 	 * @return list of successfully imported assets
 	 */
-	virtual TArray<UObject*> ImportAssets(const TArray<FString>& Files, const FString& DestinationPath, UFactory* ChosenFactory = NULL, bool bSyncToBrowser = true, TArray<TPair<FString, FString>>* FilesAndDestinations = nullptr) const = 0;
+	virtual TArray<UObject*> ImportAssets(const TArray<FString>& Files, const FString& DestinationPath, UFactory* ChosenFactory = NULL, bool bSyncToBrowser = true, TArray<TPair<FString, FString>>* FilesAndDestinations = nullptr, bool bAllowAsyncImport = false) const = 0;
 
 	/**
 	 * Imports assets using data specified completely up front.  Does not ever ask any questions of the user or show any modal error messages
diff --git a/Engine/Source/Editor/UnrealEd/Classes/Settings/EditorExperimentalSettings.h b/Engine/Source/Editor/UnrealEd/Classes/Settings/EditorExperimentalSettings.h
--- a/Engine/Source/Editor/UnrealEd/Classes/Settings/EditorExperimentalSettings.h
+++ b/Engine/Source/Editor/UnrealEd/Classes/Settings/EditorExperimentalSettings.h
@@ -36,6 +36,10 @@ public:
 	UPROPERTY(EditAnywhere, config, Category = Interchange, meta = (DisplayName = "Enable interchange framework import"))
 	bool bEnableInterchangeFramework;
 
+	/** Enable interchange framework, the interchange framework is a new import system which can import asynchronously and in parallel but only add support the texture assets. See the interchange project settings to configure the import pipeline*/
+	UPROPERTY(EditAnywhere, config, Category = Interchange, meta = (DisplayName = "Enable interchange framework import for textures only"))
+	bool bEnableInterchangeFrameworkForTextureOnly;
+
 	/** Allows the editor to run on HDR monitors on Windows 10 */
 	UPROPERTY(EditAnywhere, config, Category = HDR, meta = (ConfigRestartRequired = true, DisplayName = "Enable Editor Support for HDR Monitors"))
 	bool bHDREditor;
diff --git a/Engine/Source/Editor/UnrealEd/Private/Editor.cpp b/Engine/Source/Editor/UnrealEd/Private/Editor.cpp
--- a/Engine/Source/Editor/UnrealEd/Private/Editor.cpp
+++ b/Engine/Source/Editor/UnrealEd/Private/Editor.cpp
@@ -265,9 +265,11 @@ bool FReimportManager::Reimport( UObject* Obj, bool bAskForNewFileIfMissing, boo
 
 	bool bUseInterchangeFramework = false;
 	UInterchangeManager& InterchangeManager = UInterchangeManager::GetInterchangeManager();
-#if WITH_EDITOR
-	bUseInterchangeFramework = GetDefault<UEditorExperimentalSettings>()->bEnableInterchangeFramework;
-#endif
+	const UEditorExperimentalSettings* EditorExperimentalSettings = GetDefault<UEditorExperimentalSettings>();
+
+	bUseInterchangeFramework = EditorExperimentalSettings->bEnableInterchangeFramework;
+	const bool bUseInterchangeFrameworkForTextureOnly = (!bUseInterchangeFramework) && EditorExperimentalSettings->bEnableInterchangeFrameworkForTextureOnly;
+	bUseInterchangeFramework |= bUseInterchangeFrameworkForTextureOnly;
 
 	bool bSuccess = false;
 	if ( Obj )
@@ -390,7 +392,18 @@ bool FReimportManager::Reimport( UObject* Obj, bool bAskForNewFileIfMissing, boo
 				if (bUseInterchangeFramework)
 				{
 					UE::Interchange::FScopedSourceData ScopedSourceData(SourceFilenames[0]);
-					if (InterchangeManager.CanTranslateSourceData(ScopedSourceData.GetSourceData()))
+
+					bool bUseATextureTranslator = false;
+					if (bUseInterchangeFrameworkForTextureOnly)
+					{
+						UInterchangeTranslatorBase* Translator = InterchangeManager.GetTranslatorForSourceData(ScopedSourceData.GetSourceData());
+						if (Translator && InterchangeManager.IsTranslatorClassForTextureOnly(Translator->GetClass()))
+						{
+							bUseATextureTranslator = true;
+						}
+					}
+
+					if (bUseATextureTranslator || (!bUseInterchangeFrameworkForTextureOnly && InterchangeManager.CanTranslateSourceData(ScopedSourceData.GetSourceData())))
 					{
 						auto PostImportedLambda = [](UObject* ImportedObject)
 						{
diff --git a/Engine/Source/Editor/UnrealEd/Private/FileHelpers.cpp b/Engine/Source/Editor/UnrealEd/Private/FileHelpers.cpp
--- a/Engine/Source/Editor/UnrealEd/Private/FileHelpers.cpp
+++ b/Engine/Source/Editor/UnrealEd/Private/FileHelpers.cpp
@@ -1476,7 +1476,7 @@ void FEditorFileUtils::Import(const FString& InFilename)
 		Files.Add(InFilename);
 
 		const bool bSyncToBrowser = SceneFactory->ImportsAssets();
-		AssetToolsModule.Get().ImportAssets(Files, Path, SceneFactory, bSyncToBrowser);
+		AssetToolsModule.Get().ImportAssets(Files, Path, SceneFactory, bSyncToBrowser, nullptr, true);
 	}
 	else
 	{
diff --git a/Engine/Source/Editor/UnrealEd/Private/Subsystems/ImportSubsystem.cpp b/Engine/Source/Editor/UnrealEd/Private/Subsystems/ImportSubsystem.cpp
--- a/Engine/Source/Editor/UnrealEd/Private/Subsystems/ImportSubsystem.cpp
+++ b/Engine/Source/Editor/UnrealEd/Private/Subsystems/ImportSubsystem.cpp
@@ -94,7 +94,7 @@ public:
 			{
 				FilesAndDestinations.RemoveAt(ReImportIndexes[IndexToRemove]);
 			}
-			AssetToolsModule.Get().ImportAssets(ImportFiles, RootDestinationPath, nullptr, true, &FilesAndDestinations);
+			AssetToolsModule.Get().ImportAssets(ImportFiles, RootDestinationPath, nullptr, true, &FilesAndDestinations, true);
 		}
 	}
 
diff --git a/Engine/Source/Runtime/Experimental/Interchange/Engine/Private/InterchangeManager.cpp b/Engine/Source/Runtime/Experimental/Interchange/Engine/Private/InterchangeManager.cpp
--- a/Engine/Source/Runtime/Experimental/Interchange/Engine/Private/InterchangeManager.cpp
+++ b/Engine/Source/Runtime/Experimental/Interchange/Engine/Private/InterchangeManager.cpp
@@ -966,6 +966,19 @@ UInterchangeTranslatorBase* UInterchangeManager::GetTranslatorSupportingPayloadI
 	return nullptr;
 }
 
+void UInterchangeManager::RegisterTextureOnlyTranslatorClass(const UClass* TranslatorClass)
+{
+	if (TranslatorClass)
+	{
+		TextureOnlyTranslatorClass.Add(TranslatorClass);
+	}
+}
+
+bool UInterchangeManager::IsTranslatorClassForTextureOnly(const UClass* TranslatorClass) const
+{
+	return TextureOnlyTranslatorClass.Contains(TranslatorClass);
+}
+
 bool UInterchangeManager::IsAttended()
 {
 	if (FApp::IsGame())
diff --git a/Engine/Source/Runtime/Experimental/Interchange/Engine/Public/InterchangeManager.h b/Engine/Source/Runtime/Experimental/Interchange/Engine/Public/InterchangeManager.h
--- a/Engine/Source/Runtime/Experimental/Interchange/Engine/Public/InterchangeManager.h
+++ b/Engine/Source/Runtime/Experimental/Interchange/Engine/Public/InterchangeManager.h
@@ -449,6 +449,19 @@ public:
 	 */
 	UInterchangeTranslatorBase* GetTranslatorSupportingPayloadInterfaceForSourceData(const UInterchangeSourceData* SourceData, const UClass* PayloadInterfaceClass) const;
 
+	/**
+	 * Register a translator as a texture translator only
+	 * This used to implement the experimental setting to use interchange for textures only.
+	 * 
+	 * @Param Translator - The UClass of the translator you want to register
+	 */
+	void RegisterTextureOnlyTranslatorClass(const UClass* TranslatorClass);
+
+	/**
+	 * Check if the translator class is registered as a texture only translator
+	 */
+	bool IsTranslatorClassForTextureOnly(const UClass* TranslatorClass) const;
+
 protected:
 
 	/** Return true if we can show some UI */
@@ -536,5 +549,8 @@ private:
 	//This boolean is reset to false when the ImportTasks array is empty.
 	bool bImportAllWithDefault = false;
 
+	UPROPERTY()
+	TSet<TObjectPtr<const UClass>> TextureOnlyTranslatorClass;
+
 	friend class UE::Interchange::FScopedTranslator;
 };
