Prevented template anim BPs from being retargeted

Transitive retargeting via parents is not supported, but directly retargeting is allowed

diff --git a/Engine/Source/Developer/AssetTools/Private/AssetTypeActions/AssetTypeActions_AnimBlueprint.cpp b/Engine/Source/Developer/AssetTools/Private/AssetTypeActions/AssetTypeActions_AnimBlueprint.cpp
--- a/Engine/Source/Developer/AssetTools/Private/AssetTypeActions/AssetTypeActions_AnimBlueprint.cpp
+++ b/Engine/Source/Developer/AssetTools/Private/AssetTypeActions/AssetTypeActions_AnimBlueprint.cpp
@@ -127,18 +127,21 @@ void FAssetTypeActions_AnimBlueprint::GetActions(const TArray<UObject*>& InObjec
 				FSlateIcon(FAppStyle::GetAppStyleSetName(), "Icons.Blueprint"));
 		}
 	}
-	
-	Section.AddMenuEntry(
-		"AnimBlueprint_FindSkeleton",
-		LOCTEXT("AnimBlueprint_FindSkeleton", "Find Skeleton"),
-		LOCTEXT("AnimBlueprint_FindSkeletonTooltip", "Finds the skeleton used by the selected Anim Blueprints in the content browser."),
-		FSlateIcon(FAppStyle::GetAppStyleSetName(), "Icons.Find"),
-		FUIAction(
-			FExecuteAction::CreateSP( this, &FAssetTypeActions_AnimBlueprint::ExecuteFindSkeleton, AnimBlueprints ),
-			FCanExecuteAction()
-			)
-		);
 
+	if(AreAnyNonTemplateAnimBlueprintsSelected(AnimBlueprints))
+	{
+		Section.AddMenuEntry(
+			"AnimBlueprint_FindSkeleton",
+			LOCTEXT("AnimBlueprint_FindSkeleton", "Find Skeleton"),
+			LOCTEXT("AnimBlueprint_FindSkeletonTooltip", "Finds the skeleton used by the selected Anim Blueprints in the content browser."),
+			FSlateIcon(FAppStyle::GetAppStyleSetName(), "Icons.Find"),
+			FUIAction(
+				FExecuteAction::CreateSP( this, &FAssetTypeActions_AnimBlueprint::ExecuteFindSkeleton, AnimBlueprints ),
+				FCanExecuteAction()
+				)
+			);
+	}
+	
 	Section.AddSubMenu(
 		"RetargetBlueprintSubmenu",
 		LOCTEXT("RetargetBlueprintSubmenu", "Retarget Anim Blueprints"),
@@ -289,6 +292,19 @@ void FAssetTypeActions_AnimBlueprint::PerformAssetDiff(UObject* Asset1, UObject*
 	SBlueprintDiff::CreateDiffWindow(WindowTitle, OldBlueprint, NewBlueprint, OldRevision, NewRevision);
 }
 
+bool FAssetTypeActions_AnimBlueprint::AreAnyNonTemplateAnimBlueprintsSelected(TArray<TWeakObjectPtr<UAnimBlueprint>> Objects) const
+{
+	for(TWeakObjectPtr<UAnimBlueprint> WeakAnimBlueprint : Objects)
+	{
+		if(!WeakAnimBlueprint->bIsTemplate)
+		{
+			return true; 
+		}
+	}
+
+	return false;
+}
+
 void FAssetTypeActions_AnimBlueprint::ExecuteFindSkeleton(TArray<TWeakObjectPtr<UAnimBlueprint>> Objects)
 {
 	TArray<UObject*> ObjectsToSync;
diff --git a/Engine/Source/Developer/AssetTools/Private/AssetTypeActions/AssetTypeActions_AnimBlueprint.h b/Engine/Source/Developer/AssetTools/Private/AssetTypeActions/AssetTypeActions_AnimBlueprint.h
--- a/Engine/Source/Developer/AssetTools/Private/AssetTypeActions/AssetTypeActions_AnimBlueprint.h
+++ b/Engine/Source/Developer/AssetTools/Private/AssetTypeActions/AssetTypeActions_AnimBlueprint.h
@@ -40,4 +40,7 @@ private:
 
 	/** Handler for retargeting */
 	void RetargetAnimationHandler(USkeleton* OldSkeleton, USkeleton* NewSkeleton, bool bRemapReferencedAssets, bool bAllowRemapToExisting, bool bConvertSpaces, const EditorAnimUtils::FNameDuplicationRule* NameRule, TArray<TWeakObjectPtr<UObject>> AnimBlueprints);
+
+	/** Certain options are not available for template anim BPs */
+	bool AreAnyNonTemplateAnimBlueprintsSelected(TArray<TWeakObjectPtr<UAnimBlueprint>> Objects) const;
 };
diff --git a/Engine/Source/Editor/UnrealEd/Private/EditorAnimUtils.cpp b/Engine/Source/Editor/UnrealEd/Private/EditorAnimUtils.cpp
--- a/Engine/Source/Editor/UnrealEd/Private/EditorAnimUtils.cpp
+++ b/Engine/Source/Editor/UnrealEd/Private/EditorAnimUtils.cpp
@@ -97,16 +97,19 @@ namespace EditorAnimUtils
 			}
 			else if( UAnimBlueprint* AnimBlueprint = Cast<UAnimBlueprint>(Asset) )
 			{
-
-				// Add parents blueprint. 
+				// Add parent non-template blueprints
 				UAnimBlueprint* ParentBP = Cast<UAnimBlueprint>(AnimBlueprint->ParentClass->ClassGeneratedBy);
 				while (ParentBP)
 				{
-					AnimBlueprintsToRetarget.AddUnique(ParentBP);
+					// Cant transitively retarget templates
+					if(!(ParentBP->bIsTemplate && ParentBP->TargetSkeleton == nullptr))
+					{
+						AnimBlueprintsToRetarget.AddUnique(ParentBP);
+					}
 					ParentBP = Cast<UAnimBlueprint>(ParentBP->ParentClass->ClassGeneratedBy);
 				}
 				
-				AnimBlueprintsToRetarget.AddUnique(AnimBlueprint);				
+				AnimBlueprintsToRetarget.AddUnique(AnimBlueprint);
 			}
 		}
 		
@@ -318,6 +321,9 @@ namespace EditorAnimUtils
 
 			AnimBlueprint->TargetSkeleton = NewSkeleton;
 
+			// We can directly retarget templates (although not transitively via parents) so we need to make sure we clear the flag here
+			AnimBlueprint->bIsTemplate = false;
+
 			if (HasDuplicates())
 			{
 				// if they have parent blueprint, make sure to re-link to the new one also
