Making a rename that only change the case of an asset work properly.
It work only for the assets that aren't source controlled.

diff --git a/Engine/Source/Developer/AssetTools/Private/AssetRenameManager.cpp b/Engine/Source/Developer/AssetTools/Private/AssetRenameManager.cpp
--- a/Engine/Source/Developer/AssetTools/Private/AssetRenameManager.cpp
+++ b/Engine/Source/Developer/AssetTools/Private/AssetRenameManager.cpp
@@ -30,6 +30,8 @@
 #include "EditorStyleSet.h"
 #include "SourceControlOperations.h"
 #include "ISourceControlModule.h"
+#include "ISourceControlProvider.h"
+#include "ISourceControlState.h"
 #include "SourceControlHelpers.h"
 #include "FileHelpers.h"
 #include "SDiscoveringAssetsDialog.h"
@@ -48,6 +50,9 @@
 #include "Engine/World.h"
 #include "Engine/MapBuildDataRegistry.h"
 #include "GameMapsSettings.h"
+#include "AssetToolsModule.h"
+#include "IAssetTools.h"
+
 
 #define LOCTEXT_NAMESPACE "AssetRenameManager"
 
@@ -1412,6 +1417,7 @@ void FAssetRenameManager::PerformAssetRename(TArray<FAssetRenameDataWithReferenc
 		UObject* Asset = RenameData.Asset.Get();
 		TArray<UPackage *> PackagesToCheckForSoftReferences;
 
+		bool bIsCaseChangeOnly = false;
 		if (!RenameData.bOnlyFixSoftReferences)
 		{
 			// If bOnlyFixSoftReferences was set these got appended in find references
@@ -1428,11 +1434,65 @@ void FAssetRenameManager::PerformAssetRename(TArray<FAssetRenameDataWithReferenc
 			PGN.ObjectName = RenameData.NewName;
 			PGN.GroupName = TEXT("");
 			PGN.PackageName = RenameData.NewPackagePath / PGN.ObjectName;
-			const bool bLeaveRedirector = RenameData.bCreateRedirector;
+			bool bLeaveRedirector = RenameData.bCreateRedirector;
 
 			UPackage* OldPackage = Asset->GetOutermost();
+
+			if (OldPackage->GetFName() == FName(PGN.PackageName) && !OldPackage->IsRooted())
+			{
+				// Handle case change only.
+				bLeaveRedirector = false;
+
+				FAssetToolsModule& AssetToolsModule = FModuleManager::LoadModuleChecked<FAssetToolsModule>("AssetTools");
+				FString PackageName;
+				FString AssetName;
+				FString BasePath = RenameData.NewPackagePath + TEXT("/RenameTmp") ;
+				AssetToolsModule.Get().CreateUniqueAssetName(BasePath, TEXT(""), PackageName, AssetName);
+
+				ObjectTools::FPackageGroupName TempPGN;
+				TempPGN.ObjectName = AssetName;
+				TempPGN.GroupName = TEXT("");
+				TempPGN.PackageName = RenameData.NewPackagePath / TempPGN.ObjectName;
+
+				TSet<UPackage*> ObjectsUserRefusedToFullyLoad;
+				FText ErrorMessage;
+
+				// Case insensitive file systems and source control providers clients often handle poorly a case change.
+				bool bIsLocal = true;
+				if (ISourceControlModule::Get().IsEnabled())
+				{
+					ISourceControlProvider& Provider = ISourceControlModule::Get().GetProvider();
+					if (FSourceControlStatePtr StatePtr = Provider.GetState(Asset->GetPackage(), EStateCacheUsage::ForceUpdate))
+					{
+						if (StatePtr->IsSourceControlled())
+						{
+							bIsLocal = false;
+							ErrorMessage = LOCTEXT("ErrorCaseChangeRenameWithSourceControl", "Couldn't perform a case-only rename on a source controlled asset, as this is not supported.");
+						}
+					}
+				}
+
+				if (bIsLocal && ObjectTools::RenameSingleObject(Asset, TempPGN, ObjectsUserRefusedToFullyLoad, ErrorMessage, nullptr, bLeaveRedirector))
+				{
+					TArray<UPackage*> OldPackageToClean;
+					OldPackageToClean.Add(OldPackage);
+					OldPackage = Asset->GetPackage();
+					OldPackage->AddToRoot();
+					ObjectTools::CleanupAfterSuccessfulDelete(OldPackageToClean);
+					OldPackage->RemoveFromRoot();
+					bIsCaseChangeOnly = true;
+				}
+				else
+				{
+					RenameData.bRenameFailed = true;
+					RenameData.FailureReason = ErrorMessage;
+					continue;
+				}
+			}
+
+
 			bool bOldPackageAddedToRootSet = false;
-			if (!bLeaveRedirector && !OldPackage->IsRooted())
+			if (!bLeaveRedirector && OldPackage && !OldPackage->IsRooted())
 			{
 				bOldPackageAddedToRootSet = true;
 				OldPackage->AddToRoot();
@@ -1442,7 +1502,13 @@ void FAssetRenameManager::PerformAssetRename(TArray<FAssetRenameDataWithReferenc
 			FText ErrorMessage;
 			if (ObjectTools::RenameSingleObject(Asset, PGN, ObjectsUserRefusedToFullyLoad, ErrorMessage, nullptr, bLeaveRedirector))
 			{
-				PackagesToSave.AddUnique(Asset->GetOutermost());
+				/**
+				 * Do not save the package when the user simply changing a case and there is no referencer to it
+				 */
+				if (!(bIsCaseChangeOnly && RenameData.ReferencingPackageNames.IsEmpty()))
+				{
+					PackagesToSave.AddUnique(Asset->GetOutermost());
+				}
 
 				// Automatically save renamed assets
 				if (bLeaveRedirector)
@@ -1451,7 +1517,7 @@ void FAssetRenameManager::PerformAssetRename(TArray<FAssetRenameDataWithReferenc
 				}
 				else if (bOldPackageAddedToRootSet)
 				{
-					// Since we did not leave a redirector and the old package wasnt already rooted, attempt to delete it when we are done. 
+					// Since we did not leave a redirector and the old package wasn't already rooted, attempt to delete it when we are done. 
 					PotentialPackagesToDelete.AddUnique(OldPackage);
 				}
 			}
@@ -1469,26 +1535,29 @@ void FAssetRenameManager::PerformAssetRename(TArray<FAssetRenameDataWithReferenc
 			}
 		}
 
-		for (FName PackageName : RenameData.ReferencingPackageNames)
+		if (!RenameData.bRenameFailed && !bIsCaseChangeOnly)
 		{
-			UPackage* PackageToCheck = FindPackage(nullptr, *PackageName.ToString());
-			if (PackageToCheck)
+			for (FName PackageName : RenameData.ReferencingPackageNames)
 			{
-				PackagesToCheckForSoftReferences.Add(PackageToCheck);
+				UPackage* PackageToCheck = FindPackage(nullptr, *PackageName.ToString());
+				if (PackageToCheck)
+				{
+					PackagesToCheckForSoftReferences.Add(PackageToCheck);
+				}
 			}
-		}
 
-		TMap<FSoftObjectPath, FSoftObjectPath> RedirectorMap;
-		RedirectorMap.Add(RenameData.OldObjectPath, RenameData.NewObjectPath);
+			TMap<FSoftObjectPath, FSoftObjectPath> RedirectorMap;
+			RedirectorMap.Add(RenameData.OldObjectPath, RenameData.NewObjectPath);
 
-		if (UBlueprint* Blueprint = Cast<UBlueprint>(Asset))
-		{
-			// Add redirect for class and default as well
-			RedirectorMap.Add(FString::Printf(TEXT("%s_C"), *RenameData.OldObjectPath.ToString()), FString::Printf(TEXT("%s_C"), *RenameData.NewObjectPath.ToString()));
-			RedirectorMap.Add(FString::Printf(TEXT("%s.Default__%s_C"), *RenameData.OldObjectPath.GetLongPackageName(), *RenameData.OldObjectPath.GetAssetName()), FString::Printf(TEXT("%s.Default__%s_C"), *RenameData.NewObjectPath.GetLongPackageName(), *RenameData.NewObjectPath.GetAssetName()));
-		}
+			if (UBlueprint* Blueprint = Cast<UBlueprint>(Asset))
+			{
+				// Add redirect for class and default as well
+				RedirectorMap.Add(FString::Printf(TEXT("%s_C"), *RenameData.OldObjectPath.ToString()), FString::Printf(TEXT("%s_C"), *RenameData.NewObjectPath.ToString()));
+				RedirectorMap.Add(FString::Printf(TEXT("%s.Default__%s_C"), *RenameData.OldObjectPath.GetLongPackageName(), *RenameData.OldObjectPath.GetAssetName()), FString::Printf(TEXT("%s.Default__%s_C"), *RenameData.NewObjectPath.GetLongPackageName(), *RenameData.NewObjectPath.GetAssetName()));
+			}
 
-		RenameReferencingSoftObjectPaths(PackagesToCheckForSoftReferences, RedirectorMap);
+			RenameReferencingSoftObjectPaths(PackagesToCheckForSoftReferences, RedirectorMap);
+		}
 	}
 
 	GWarn->EndSlowTask();
diff --git a/Engine/Source/Editor/ContentBrowserData/Public/ContentBrowserItemData.h b/Engine/Source/Editor/ContentBrowserData/Public/ContentBrowserItemData.h
--- a/Engine/Source/Editor/ContentBrowserData/Public/ContentBrowserItemData.h
+++ b/Engine/Source/Editor/ContentBrowserData/Public/ContentBrowserItemData.h
@@ -612,7 +612,7 @@ public:
 	bool operator==(const FContentBrowserItemDataKey& InOther) const
 	{
 		return ItemType == InOther.ItemType
-			&& VirtualPath == InOther.VirtualPath;
+			&& VirtualPath.IsEqual(InOther.VirtualPath, ENameCase::CaseSensitive);
 	}
 
 	/**
diff --git a/Engine/Source/Editor/UnrealEd/Private/ObjectTools.cpp b/Engine/Source/Editor/UnrealEd/Private/ObjectTools.cpp
--- a/Engine/Source/Editor/UnrealEd/Private/ObjectTools.cpp
+++ b/Engine/Source/Editor/UnrealEd/Private/ObjectTools.cpp
@@ -3690,48 +3690,60 @@ namespace ObjectTools
 				{
 					// We can rename on top of an object redirection (basically destroy the redirection and put us in its place).
 					UPackage* NewPackage = CreatePackage( *FullPackageName );
-					NewPackage->GetOutermost()->FullyLoad();
 
-					// Make sure we copy all the cooked package flags if the asset was already cooked.
-					if (Object->GetOutermost()->HasAnyPackageFlags(PKG_FilterEditorOnly))
-					{
-						NewPackage->SetPackageFlags(PKG_FilterEditorOnly);
-					}
-					NewPackage->bIsCookedForEditor = Object->GetOutermost()->bIsCookedForEditor;
+					bool bFoundCompatibleRedirector = false;
+					UObjectRedirector* Redirector = nullptr;
 
-					// Renaming an asset should respect the export controls of the original.
-					if (Object->GetOutermost()->HasAnyPackageFlags(PKG_DisallowExport))
+					if (NewPackage != Object->GetPackage())
 					{
-						NewPackage->SetPackageFlags(PKG_DisallowExport);
-					}
+						NewPackage->GetOutermost()->FullyLoad();
 
-					// When renaming a World Composition map, make sure to properly initialize WorldTileInfo
-					if (Object->GetOutermost()->WorldTileInfo.IsValid())
-					{
-						NewPackage->WorldTileInfo = MakeUnique<FWorldTileInfo>(*Object->GetOutermost()->WorldTileInfo);
-					}
+						// Make sure we copy all the cooked package flags if the asset was already cooked.
+						if (Object->GetOutermost()->HasAnyPackageFlags(PKG_FilterEditorOnly))
+						{
+							NewPackage->SetPackageFlags(PKG_FilterEditorOnly);
+						}
+						NewPackage->bIsCookedForEditor = Object->GetOutermost()->bIsCookedForEditor;
 
-					UObjectRedirector* Redirector = Cast<UObjectRedirector>( StaticFindObject(UObjectRedirector::StaticClass(), NewPackage, *NewObjectName) );
-					bool bFoundCompatibleRedirector = false;
-					// If we found a redirector, check that the object it points to is of the same class.
-					if ( Redirector
-						&& Redirector->DestinationObject
-						&& Redirector->DestinationObject->GetClass() == Object->GetClass() )
-					{
-						// Test renaming the redirector into a dummy package.
-						if ( Redirector->Rename(*Redirector->GetName(), CreatePackage( TEXT("/Temp/TempRedirectors")), REN_Test) )
+						// Renaming an asset should respect the export controls of the original.
+						if (Object->GetOutermost()->HasAnyPackageFlags(PKG_DisallowExport))
 						{
-							// Actually rename the redirector here so it doesn't get in the way of the rename below.
-							Redirector->Rename(*Redirector->GetName(), CreatePackage( TEXT("/Temp/TempRedirectors")), REN_DontCreateRedirectors);
+							NewPackage->SetPackageFlags(PKG_DisallowExport);
+						}
 
-							bFoundCompatibleRedirector = true;
+						// When renaming a World Composition map, make sure to properly initialize WorldTileInfo
+						if (Object->GetOutermost()->WorldTileInfo.IsValid())
+						{
+							NewPackage->WorldTileInfo = MakeUnique<FWorldTileInfo>(*Object->GetOutermost()->WorldTileInfo);
 						}
-						else
+
+						Redirector = Cast<UObjectRedirector>(StaticFindObject(UObjectRedirector::StaticClass(), NewPackage, *NewObjectName));
+						// If we found a redirector, check that the object it points to is of the same class.
+						if (Redirector
+							&& Redirector->DestinationObject
+							&& Redirector->DestinationObject->GetClass() == Object->GetClass())
 						{
-							bMoveFailed = true;
-							bMoveRedirectorFailed = true;
+							// Test renaming the redirector into a dummy package.
+							if (Redirector->Rename(*Redirector->GetName(), CreatePackage(TEXT("/Temp/TempRedirectors")), REN_Test))
+							{
+								// Actually rename the redirector here so it doesn't get in the way of the rename below.
+								Redirector->Rename(*Redirector->GetName(), CreatePackage(TEXT("/Temp/TempRedirectors")), REN_DontCreateRedirectors);
+
+								bFoundCompatibleRedirector = true;
+							}
+							else
+							{
+								bMoveFailed = true;
+								bMoveRedirectorFailed = true;
+							}
 						}
 					}
+					else
+					{
+						bMoveFailed = true;
+						ErrorMessage += (NSLOCTEXT("UnrealEd", "Error_ObjectNameCaseChange", "Cannot change the case of an object name.\n")).ToString();
+					}
+
 
 					if ( !bMoveFailed )
 					{
