Fix exception when right-clicking a child of a template AnimBP that has no skeleton

diff --git a/Engine/Source/Developer/AssetTools/Private/AssetTypeActions/AssetTypeActions_AnimBlueprint.cpp b/Engine/Source/Developer/AssetTools/Private/AssetTypeActions/AssetTypeActions_AnimBlueprint.cpp
--- a/Engine/Source/Developer/AssetTools/Private/AssetTypeActions/AssetTypeActions_AnimBlueprint.cpp
+++ b/Engine/Source/Developer/AssetTools/Private/AssetTypeActions/AssetTypeActions_AnimBlueprint.cpp
@@ -39,7 +39,7 @@ void FAssetTypeActions_AnimBlueprint::GetActions(const TArray<UObject*>& InObjec
 		UAnimBlueprint* AnimBlueprint = AnimBlueprints[0].Get();
 
 		// Accept (non-interface) template anim BPs or anim BPs with compatible skeletons
-		if(AnimBlueprint && AnimBlueprint->BlueprintType != BPTYPE_Interface && ((AnimBlueprint->TargetSkeleton == nullptr && AnimBlueprint->bIsTemplate) || AnimBlueprint->TargetSkeleton->GetCompatibleSkeletons().Num() > 0))
+		if(AnimBlueprint && AnimBlueprint->BlueprintType != BPTYPE_Interface && ((AnimBlueprint->TargetSkeleton == nullptr && AnimBlueprint->bIsTemplate) || (AnimBlueprint->TargetSkeleton != nullptr && AnimBlueprint->TargetSkeleton->GetCompatibleSkeletons().Num() > 0)))
 		{
 			Section.AddSubMenu(
 				"AnimBlueprint_NewSkeletonChildBlueprint",
@@ -325,6 +325,11 @@ void FAssetTypeActions_AnimBlueprint::ExecuteFindSkeleton(TArray<TWeakObjectPtr<
 	{
 		FAssetTools::Get().SyncBrowserToAssets(ObjectsToSync);
 	}
+	else
+	{
+		FText ShouldRetargetMessage = LOCTEXT("NoSkeletonFound", "Could not find the skeleton");
+		FMessageDialog::Open(EAppMsgType::Ok, ShouldRetargetMessage);
+	}
 }
 
 void FAssetTypeActions_AnimBlueprint::RetargetAnimationHandler(USkeleton* OldSkeleton, USkeleton* NewSkeleton, bool bRemapReferencedAssets, bool bAllowRemapToExisting, bool bConvertSpaces, const EditorAnimUtils::FNameDuplicationRule* NameRule, TArray<TWeakObjectPtr<UObject>> AnimBlueprints)
