If Interchange framework is enable we will force content browser async import
#rnx

diff --git a/Engine/Source/Developer/AssetTools/Private/AssetTools.cpp b/Engine/Source/Developer/AssetTools/Private/AssetTools.cpp
--- a/Engine/Source/Developer/AssetTools/Private/AssetTools.cpp
+++ b/Engine/Source/Developer/AssetTools/Private/AssetTools.cpp
@@ -1172,7 +1172,13 @@ TArray<UObject*> UAssetToolsImpl::ImportAssetsWithDialog(const FString& Destinat
 
 
 			FEditorDirectories::Get().SetLastDirectory(ELastDirectory::GENERIC_IMPORT, OpenFilenames[0]);
-			ReturnObjects = ImportAssets(OpenFilenames, DestinationPath, ChosenFactory);
+
+			//If interchange framework is enable and we do not import using automation we allow async import.
+			const UEditorExperimentalSettings* EditorExperimentalSettings = GetDefault<UEditorExperimentalSettings>();
+			const bool bAllowAsyncImport = EditorExperimentalSettings->bEnableInterchangeFramework;
+			constexpr bool bSyncToBrowser = true;
+			constexpr TArray<TPair<FString, FString>>* FilesAndDestinations = nullptr;
+			ReturnObjects = ImportAssets(OpenFilenames, DestinationPath, ChosenFactory, bSyncToBrowser, FilesAndDestinations, bAllowAsyncImport);
 		}
 	}
 
