Fix crash because of GCed empty packages

diff --git a/Engine/Source/Developer/AssetTools/Private/AssetRenameManager.cpp b/Engine/Source/Developer/AssetTools/Private/AssetRenameManager.cpp
--- a/Engine/Source/Developer/AssetTools/Private/AssetRenameManager.cpp
+++ b/Engine/Source/Developer/AssetTools/Private/AssetRenameManager.cpp
@@ -52,6 +52,7 @@
 #include "GameMapsSettings.h"
 #include "AssetToolsModule.h"
 #include "IAssetTools.h"
+#include "Algo/Transform.h"
 
 
 #define LOCTEXT_NAMESPACE "AssetRenameManager"
@@ -1568,9 +1569,15 @@ void FAssetRenameManager::PerformAssetRename(TArray<FAssetRenameDataWithReferenc
 		const bool bCheckDirty = false;
 		const bool bPromptToSave = false;
 		const bool bAlreadyCheckedOut = true;
+
+		// Get the list of filenames before calling save because some of the saved packages can get GCed if they are empty packages
+		TArray<FString> Filenames;
+		Filenames.Reserve(PackagesToSave.Num());
+		Algo::Transform(PackagesToSave, Filenames, [](UPackage* InPackage) { return USourceControlHelpers::PackageFilename(InPackage); });
+
 		FEditorFileUtils::PromptForCheckoutAndSave(PackagesToSave, bCheckDirty, bPromptToSave, nullptr, bAlreadyCheckedOut);
 
-		ISourceControlModule::Get().QueueStatusUpdate(PackagesToSave);
+		ISourceControlModule::Get().QueueStatusUpdate(Filenames);
 	}
 
 	// Bulk update SCC status for old packages since it is faster than doing it one by one below
