Fix for editor crash in PerformDuplicateAsset() caused by PackageName string with "//" in path being passed to CanCreateAsset().

Specific crash happened when PackageName string was created from PackagePath + TEXT("/") + AssetName, and PackagePath ended with "/" char; resulting in a package name with "//" in the body. This specific condition is checked for in CreatePackage() and when found results in a fatal log error being thrown.

Fixed by changing PackageName string to build using the / operator which prevents "//" in the result string

diff --git a/Engine/Source/Developer/AssetTools/Private/AssetTools.cpp b/Engine/Source/Developer/AssetTools/Private/AssetTools.cpp
--- a/Engine/Source/Developer/AssetTools/Private/AssetTools.cpp
+++ b/Engine/Source/Developer/AssetTools/Private/AssetTools.cpp
@@ -632,7 +632,7 @@ UObject* UAssetToolsImpl::PerformDuplicateAsset(const FString& AssetName, const
 		return nullptr;
 	}
 
-	const FString PackageName = PackagePath + TEXT("/") + AssetName;
+	const FString PackageName = PackagePath / AssetName;
 
 	// Make sure we can create the asset without conflicts
 	if ( !CanCreateAsset(AssetName, PackageName, LOCTEXT("DuplicateAnObject", "Duplicate an object")) )
