Added a function ImportAssetsWithDialogAsync to the AssetTools.

Changed two import calls in the editor where interchange wasn't used. ImportAssetsWithDialog doesn't support async import to avoid breaking existing automation scripts or code.

diff --git a/Engine/Plugins/Editor/ContentBrowser/ContentBrowserAssetDataSource/Source/ContentBrowserAssetDataSource/Private/ContentBrowserAssetDataSource.cpp b/Engine/Plugins/Editor/ContentBrowser/ContentBrowserAssetDataSource/Source/ContentBrowserAssetDataSource/Private/ContentBrowserAssetDataSource.cpp
--- a/Engine/Plugins/Editor/ContentBrowser/ContentBrowserAssetDataSource/Source/ContentBrowserAssetDataSource/Private/ContentBrowserAssetDataSource.cpp
+++ b/Engine/Plugins/Editor/ContentBrowser/ContentBrowserAssetDataSource/Source/ContentBrowserAssetDataSource/Private/ContentBrowserAssetDataSource.cpp
@@ -2026,7 +2026,7 @@ void UContentBrowserAssetDataSource::OnImportAsset(const FName InPath)
 {
 	if (ensure(!InPath.IsNone()))
 	{
-		AssetTools->ImportAssetsWithDialog(InPath.ToString());
+		AssetTools->ImportAssetsWithDialogAsync(InPath.ToString());
 	}
 }
 
diff --git a/Engine/Source/Developer/AssetTools/Private/AssetTools.cpp b/Engine/Source/Developer/AssetTools/Private/AssetTools.cpp
--- a/Engine/Source/Developer/AssetTools/Private/AssetTools.cpp
+++ b/Engine/Source/Developer/AssetTools/Private/AssetTools.cpp
@@ -1106,83 +1106,14 @@
 
 TArray<UObject*> UAssetToolsImpl::ImportAssetsWithDialog(const FString& DestinationPath)
 {
-	if (!GetWritableFolderPermissionList()->PassesStartsWithFilter(DestinationPath))
-	{
-		NotifyBlockedByWritableFolderFilter();
-		return TArray<UObject*>();
-	}
-
-	TArray<UObject*> ReturnObjects;
-	FString FileTypes, AllExtensions;
-	TArray<UFactory*> Factories;
-
-	// Get the list of valid factories
-	for( TObjectIterator<UClass> It ; It ; ++It )
-	{
-		UClass* CurrentClass = (*It);
-
-		if( CurrentClass->IsChildOf(UFactory::StaticClass()) && !(CurrentClass->HasAnyClassFlags(CLASS_Abstract)) )
-		{
-			UFactory* Factory = Cast<UFactory>( CurrentClass->GetDefaultObject() );
-			if( Factory->bEditorImport )
-			{
-				Factories.Add( Factory );
-			}
-		}
-	}
-
-	TMultiMap<uint32, UFactory*> FilterIndexToFactory;
-
-	// Generate the file types and extensions represented by the selected factories
-	ObjectTools::GenerateFactoryFileExtensions( Factories, FileTypes, AllExtensions, FilterIndexToFactory );
-
-	FileTypes = FString::Printf(TEXT("All Files (%s)|%s|%s"),*AllExtensions,*AllExtensions,*FileTypes);
-
-	// Prompt the user for the filenames
-	TArray<FString> OpenFilenames;
-	IDesktopPlatform* DesktopPlatform = FDesktopPlatformModule::Get();
-	bool bOpened = false;
-	int32 FilterIndex = -1;
-
-	if ( DesktopPlatform )
-	{
-		const void* ParentWindowWindowHandle = FSlateApplication::Get().FindBestParentWindowHandleForDialogs(nullptr);
-
-		bOpened = DesktopPlatform->OpenFileDialog(
-			ParentWindowWindowHandle,
-			LOCTEXT("ImportDialogTitle", "Import").ToString(),
-			FEditorDirectories::Get().GetLastDirectory(ELastDirectory::GENERIC_IMPORT),
-			TEXT(""),
-			FileTypes,
-			EFileDialogFlags::Multiple,
-			OpenFilenames,
-			FilterIndex
-			);
-	}
-
-	if ( bOpened )
-	{
-		if ( OpenFilenames.Num() > 0 )
-		{
-			UFactory* ChosenFactory = nullptr;
-			if (FilterIndex > 0)
-			{
-				ChosenFactory = *FilterIndexToFactory.Find(FilterIndex);
-			}
-
-
-			FEditorDirectories::Get().SetLastDirectory(ELastDirectory::GENERIC_IMPORT, OpenFilenames[0]);
-
-			//If interchange framework is enable and we do not import using automation we allow async import.
-			const UEditorExperimentalSettings* EditorExperimentalSettings = GetDefault<UEditorExperimentalSettings>();
-			const bool bAllowAsyncImport = EditorExperimentalSettings->bEnableInterchangeFramework;
-			constexpr bool bSyncToBrowser = true;
-			constexpr TArray<TPair<FString, FString>>* FilesAndDestinations = nullptr;
-			ReturnObjects = ImportAssets(OpenFilenames, DestinationPath, ChosenFactory, bSyncToBrowser, FilesAndDestinations, bAllowAsyncImport);
-		}
-	}
-
-	return ReturnObjects;
+	const bool bAllowAsyncImport = false;
+	return ImportAssetsWithDialogImplementation(DestinationPath, bAllowAsyncImport);
 }
 
+void UAssetToolsImpl::ImportAssetsWithDialogAsync(const FString& DestinationPath)
+{
+	const bool bAllowAsyncImport = true;
+	ImportAssetsWithDialogImplementation(DestinationPath, bAllowAsyncImport);
+}
+
 TArray<UObject*> UAssetToolsImpl::ImportAssetsAutomated(const UAutomatedAssetImportData* ImportData)
@@ -3568,6 +3499,83 @@
 	AddSubContentBlacklist(InAssetPath);
 }
 
+TArray<UObject*> UAssetToolsImpl::ImportAssetsWithDialogImplementation(const FString& DestinationPath, bool bAllowAsyncImport)
+{
+	if (!GetWritableFolderPermissionList()->PassesStartsWithFilter(DestinationPath))
+	{
+		NotifyBlockedByWritableFolderFilter();
+		return TArray<UObject*>();
+	}
+
+	TArray<UObject*> ReturnObjects;
+	FString FileTypes, AllExtensions;
+	TArray<UFactory*> Factories;
+
+	// Get the list of valid factories
+	for (TObjectIterator<UClass> It; It; ++It)
+	{
+		UClass* CurrentClass = (*It);
+
+		if (CurrentClass->IsChildOf(UFactory::StaticClass()) && !(CurrentClass->HasAnyClassFlags(CLASS_Abstract)))
+		{
+			UFactory* Factory = Cast<UFactory>(CurrentClass->GetDefaultObject());
+			if (Factory->bEditorImport)
+			{
+				Factories.Add(Factory);
+			}
+		}
+	}
+
+	TMultiMap<uint32, UFactory*> FilterIndexToFactory;
+
+	// Generate the file types and extensions represented by the selected factories
+	ObjectTools::GenerateFactoryFileExtensions(Factories, FileTypes, AllExtensions, FilterIndexToFactory);
+
+	FileTypes = FString::Printf(TEXT("All Files (%s)|%s|%s"), *AllExtensions, *AllExtensions, *FileTypes);
+
+	// Prompt the user for the filenames
+	TArray<FString> OpenFilenames;
+	IDesktopPlatform* DesktopPlatform = FDesktopPlatformModule::Get();
+	bool bOpened = false;
+	int32 FilterIndex = -1;
+
+	if (DesktopPlatform)
+	{
+		const void* ParentWindowWindowHandle = FSlateApplication::Get().FindBestParentWindowHandleForDialogs(nullptr);
+
+		bOpened = DesktopPlatform->OpenFileDialog(
+			ParentWindowWindowHandle,
+			LOCTEXT("ImportDialogTitle", "Import").ToString(),
+			FEditorDirectories::Get().GetLastDirectory(ELastDirectory::GENERIC_IMPORT),
+			TEXT(""),
+			FileTypes,
+			EFileDialogFlags::Multiple,
+			OpenFilenames,
+			FilterIndex
+		);
+	}
+
+	if (bOpened)
+	{
+		if (OpenFilenames.Num() > 0)
+		{
+			UFactory* ChosenFactory = nullptr;
+			if (FilterIndex > 0)
+			{
+				ChosenFactory = *FilterIndexToFactory.Find(FilterIndex);
+			}
+
+
+			FEditorDirectories::Get().SetLastDirectory(ELastDirectory::GENERIC_IMPORT, OpenFilenames[0]);
+			const bool bSyncToBrowser = false;
+			TArray<TPair<FString, FString>>* FilesAndDestination = nullptr;
+			ReturnObjects = ImportAssets(OpenFilenames, DestinationPath, ChosenFactory, bSyncToBrowser, FilesAndDestination, bAllowAsyncImport);
+		}
+	}
+
+	return ReturnObjects;
+}
+
 TSharedRef<FPathPermissionList>& UAssetToolsImpl::GetFolderPermissionList()
 {
 	return FolderPermissionList;
diff --git a/Engine/Source/Developer/AssetTools/Private/AssetTools.h b/Engine/Source/Developer/AssetTools/Private/AssetTools.h
--- a/Engine/Source/Developer/AssetTools/Private/AssetTools.h
+++ b/Engine/Source/Developer/AssetTools/Private/AssetTools.h
@@ -88,6 +88,7 @@ public:
 	virtual void RenameReferencingSoftObjectPaths(const TArray<UPackage *> PackagesToCheck, const TMap<FSoftObjectPath, FSoftObjectPath>& AssetRedirectorMap) override;
 	virtual TArray<UObject*> ImportAssets(const FString& DestinationPath) override;
 	virtual TArray<UObject*> ImportAssetsWithDialog(const FString& DestinationPath) override;
+	virtual void ImportAssetsWithDialogAsync(const FString& DestinationPath) override;
 	virtual TArray<UObject*> ImportAssets(const TArray<FString>& Files, const FString& DestinationPath, UFactory* ChosenFactory, bool bSyncToBrowser = true, TArray<TPair<FString, FString>>* FilesAndDestinations = nullptr, bool bAllowAsyncImport = false) const override;
 	virtual TArray<UObject*> ImportAssetsAutomated(const UAutomatedAssetImportData* ImportData) override;
 	virtual void ImportAssetTasks(const TArray<UAssetImportTask*>& ImportTasks) override;
@@ -191,6 +192,9 @@ private:
 	/** Called when a new mount is added to add the proper sub content blacklist to it. */
 	void OnContentPathMounted(const FString& InAssetPath, const FString& FileSystemPath);
 
+	/** Implementation for the import with dialog functions */
+	TArray<UObject*> ImportAssetsWithDialogImplementation(const FString& DestinationPath, bool bAllowAsyncImport);
+
 private:
 	/** The list of all registered AssetTypeActions */
 	TArray<TSharedRef<IAssetTypeActions>> AssetTypeActionsList;
diff --git a/Engine/Source/Developer/AssetTools/Public/IAssetTools.h b/Engine/Source/Developer/AssetTools/Public/IAssetTools.h
--- a/Engine/Source/Developer/AssetTools/Public/IAssetTools.h
+++ b/Engine/Source/Developer/AssetTools/Public/IAssetTools.h
@@ -286,6 +286,14 @@ public:
 	UFUNCTION(BlueprintCallable, Category = "Editor Scripting | Asset Tools")
 	virtual TArray<UObject*> ImportAssetsWithDialog(const FString& DestinationPath) = 0;
 
+	/**
+	 * Opens a file open dialog to choose files to import to the destination path.
+	 * It differ from ImportAssetsWithDialog by allowing an async import when it's available
+	 *
+	 * @param DestinationPath	Path to import files to
+	 */
+	virtual void ImportAssetsWithDialogAsync(const FString& DestinationPath) = 0;
+
 	/**
 	 * Imports the specified files to the destination path.
 	 *
diff --git a/Engine/Source/Editor/LevelEditor/Private/LevelEditorActions.cpp b/Engine/Source/Editor/LevelEditor/Private/LevelEditorActions.cpp
--- a/Engine/Source/Editor/LevelEditor/Private/LevelEditorActions.cpp
+++ b/Engine/Source/Editor/LevelEditor/Private/LevelEditorActions.cpp
@@ -2277,7 +2277,7 @@ void FLevelEditorActionCallbacks::ImportContent()
 
 	FAssetToolsModule& AssetToolsModule = FModuleManager::Get().LoadModuleChecked<FAssetToolsModule>(TEXT("AssetTools"));
 
-	AssetToolsModule.Get().ImportAssetsWithDialog(Path);
+	AssetToolsModule.Get().ImportAssetsWithDialogAsync(Path);
 }
 
 void FLevelEditorActionCallbacks::ToggleVR()
