[AssetTools] Replace/remove uses of whitelist/blacklist

diff --git a/Engine/Config/BaseEditor.ini b/Engine/Config/BaseEditor.ini
--- a/Engine/Config/BaseEditor.ini
+++ b/Engine/Config/BaseEditor.ini
@@ -23,9 +23,9 @@ ColorFontAtlasSize=512
 
 [AssetTools]
 ; Full Package Path being hidden from the Asset view
-;+BlacklistAssetPaths="/Game/__ExternalActors__"
+;+DenyListAssetPaths="/Game/__ExternalActors__"
 ; Content Sub Path being hidden from the Asset view for each mount
-;+BlacklistContentSubPaths=
+;+DenyListContentSubPaths=
 
 [/Script/UnrealEd.ThumbnailManager]
 +RenderableThumbnailTypes=(ClassNeedingThumbnailName="/Script/CoreUObject.Class",RendererClassName="/Script/UnrealEd.ClassThumbnailRenderer")
diff --git a/Engine/Source/Developer/AssetTools/Private/AssetTools.cpp b/Engine/Source/Developer/AssetTools/Private/AssetTools.cpp
--- a/Engine/Source/Developer/AssetTools/Private/AssetTools.cpp
+++ b/Engine/Source/Developer/AssetTools/Private/AssetTools.cpp
@@ -186,19 +186,19 @@ UAssetToolsImpl::UAssetToolsImpl(const FObjectInitializer& ObjectInitializer)
 	}
 	AssetClassPermissionList->OnFilterChanged().AddUObject(this, &UAssetToolsImpl::AssetClassPermissionListChanged);
 
-	TArray<FString> BlacklistedViewPath;
-	GConfig->GetArray(TEXT("AssetTools"), TEXT("BlacklistAssetPaths"), BlacklistedViewPath, GEditorIni);
-	for (const FString& Path : BlacklistedViewPath)
+	TArray<FString> DenyListedViewPath;
+	GConfig->GetArray(TEXT("AssetTools"), TEXT("DenyListAssetPaths"), DenyListedViewPath, GEditorIni);
+	for (const FString& Path : DenyListedViewPath)
 	{
 		FolderPermissionList->AddDenyListItem("AssetToolsConfigFile", Path);
 	}
 
-	GConfig->GetArray(TEXT("AssetTools"), TEXT("BlacklistContentSubPaths"), SubContentBlacklistPaths, GEditorIni);
+	GConfig->GetArray(TEXT("AssetTools"), TEXT("DenyListContentSubPaths"), SubContentDenyListPaths, GEditorIni);
 	TArray<FString> ContentRoots;
 	FPackageName::QueryRootContentPaths(ContentRoots);
 	for (const FString& ContentRoot : ContentRoots)
 	{
-		AddSubContentBlacklist(ContentRoot);
+		AddSubContentDenyList(ContentRoot);
 	}
 	FPackageName::OnContentPathMounted().AddUObject(this, &UAssetToolsImpl::OnContentPathMounted);
 
@@ -3122,7 +3122,7 @@ void UAssetToolsImpl::RecursiveGetDependencies(const FName& PackageName, TSet<FN
 					if (!ExternalObjectsPath.IsEmpty() && !OutExternalObjectsPaths.Contains(ExternalObjectsPath))
 					{
 						OutExternalObjectsPaths.Add(ExternalObjectsPath);
-						AssetRegistryModule.Get().ScanPathsSynchronous({ ExternalObjectsPath }, /*bForceRescan*/true, /*bIgnoreBlackListScanFilters*/true);
+						AssetRegistryModule.Get().ScanPathsSynchronous({ ExternalObjectsPath }, /*bForceRescan*/true, /*bIgnoreDenyListScanFilters*/true);
 
 						TArray<FAssetData> ExternalObjectAssets;
 						AssetRegistryModule.Get().GetAssetsByPath(FName(*ExternalObjectsPath), ExternalObjectAssets, /*bRecursive*/true);
@@ -3491,9 +3491,9 @@ void UAssetToolsImpl::AssetClassPermissionListChanged()
 	}
 }
 
-void UAssetToolsImpl::AddSubContentBlacklist(const FString& InMount)
+void UAssetToolsImpl::AddSubContentDenyList(const FString& InMount)
 {
-	for (const FString& SubContentPath : SubContentBlacklistPaths)
+	for (const FString& SubContentPath : SubContentDenyListPaths)
 	{
 		FolderPermissionList->AddDenyListItem("AssetToolsConfigFile", InMount / SubContentPath);
 	}
@@ -3501,7 +3501,7 @@ void UAssetToolsImpl::AddSubContentBlacklist(const FString& InMount)
 
 void UAssetToolsImpl::OnContentPathMounted(const FString& InAssetPath, const FString& FileSystemPath)
 {
-	AddSubContentBlacklist(InAssetPath);
+	AddSubContentDenyList(InAssetPath);
 }
 
 TArray<UObject*> UAssetToolsImpl::ImportAssetsWithDialogImplementation(const FString& DestinationPath, bool bAllowAsyncImport)
diff --git a/Engine/Source/Developer/AssetTools/Private/AssetTools.h b/Engine/Source/Developer/AssetTools/Private/AssetTools.h
--- a/Engine/Source/Developer/AssetTools/Private/AssetTools.h
+++ b/Engine/Source/Developer/AssetTools/Private/AssetTools.h
@@ -180,16 +180,16 @@ private:
 
 	UObject* PerformDuplicateAsset(const FString& AssetName, const FString& PackagePath, UObject* OriginalObject, bool bWithDialog);
 
-	/** Internal method that performs actions when asset class blacklist filter changes */
+	/** Internal method that performs actions when asset class deny list filter changes */
 	void AssetClassPermissionListChanged();
 
 	/**
-	 * Add sub content blacklist filter for a new mount point
+	 * Add sub content deny list filter for a new mount point
 	 * @param InMount The mount point
 	 */
-	void AddSubContentBlacklist(const FString& InMount);
+	void AddSubContentDenyList(const FString& InMount);
 
-	/** Called when a new mount is added to add the proper sub content blacklist to it. */
+	/** Called when a new mount is added to add the proper sub content deny list to it. */
 	void OnContentPathMounted(const FString& InAssetPath, const FString& FileSystemPath);
 
 	/** Implementation for the import with dialog functions */
@@ -218,7 +218,7 @@ private:
 	TSharedRef<FPathPermissionList> WritableFolderPermissionList;
 
 	/** List of sub content paths denied for every mount. */
-	TArray<FString> SubContentBlacklistPaths;
+	TArray<FString> SubContentDenyListPaths;
 };
 
 PRAGMA_ENABLE_DEPRECATION_WARNINGS
