Exposed IAssetsTools::DiffAgainstDepot() and IAssetTools::DiffAssets()  functions to Python

diff --git a/Engine/Source/Developer/AssetTools/Public/IAssetTools.h b/Engine/Source/Developer/AssetTools/Public/IAssetTools.h
--- a/Engine/Source/Developer/AssetTools/Public/IAssetTools.h
+++ b/Engine/Source/Developer/AssetTools/Public/IAssetTools.h
@@ -364,9 +364,11 @@ public:
 	 * @param InPackagePath - The fullpath to the package
 	 * @param InPackageName - The name of the package
 	 */
+	UFUNCTION(BlueprintCallable, Category = "Editor Scripting | Asset Tools")
 	virtual void DiffAgainstDepot(UObject* InObject, const FString& InPackagePath, const FString& InPackageName) const = 0;
 
 	/** Try and diff two assets using class-specific tool. Will do nothing if either asset is NULL, or they are not the same class. */
+	UFUNCTION(BlueprintCallable, Category = "Editor Scripting | Asset Tools")
 	virtual void DiffAssets(UObject* OldAsset, UObject* NewAsset, const struct FRevisionInfo& OldRevision, const struct FRevisionInfo& NewRevision) const = 0;
 
 	/** Util for dumping an asset to a temporary text file. Returns absolute filename to temp file */
diff --git a/Engine/Source/Developer/AssetTools/Public/IAssetTypeActions.h b/Engine/Source/Developer/AssetTools/Public/IAssetTypeActions.h
--- a/Engine/Source/Developer/AssetTools/Public/IAssetTypeActions.h
+++ b/Engine/Source/Developer/AssetTools/Public/IAssetTypeActions.h
@@ -7,6 +7,7 @@
 #include "Developer/Merge/Public/Merge.h"
 #include "ThumbnailRendering/ThumbnailManager.h"
 #include "AssetRegistry/AssetData.h"
+#include "IAssetTypeActions.generated.h"
 
 class IToolkitHost;
 
@@ -23,11 +24,19 @@ namespace EAssetTypeActivationMethod
 class IToolkitHost;
 
 /* Revision information for a single revision of a file in source control */
+USTRUCT(BlueprintType)
 struct FRevisionInfo
 {
+	GENERATED_BODY()
+
+	UPROPERTY(BlueprintReadWrite, Category="Asset Revision")
 	FString		Revision;
-	int32		Changelist;
-	FDateTime	Date;	
+
+	UPROPERTY(BlueprintReadWrite, Category="Asset Revision")
+	int32		Changelist = -1;
+
+	UPROPERTY(BlueprintReadWrite, Category="Asset Revision")
+	FDateTime	Date;
 
 	static inline FRevisionInfo InvalidRevision()
 	{
