ContentBrowser: Better support for Delete/Copy/Rename operations on World Partition worlds

- FScopedLoadAllExternalObjects object that can be used to force the loading of all external actors of a WP world to allow those operations to succeed. (for larger worlds commandlet should be used)
- ContentBrowser operation like: Delete, Duplicate, Rename, Move, Copy use this new scope
- World Partition worlds are no longer loaded on right-click context menu.
- Note: change was done with explicit scopes to only address specific use cases and limit the impact of this CL. It could become a default behavior in later versions.

diff --git a/Engine/Plugins/Editor/ContentBrowser/ContentBrowserAssetDataSource/Source/ContentBrowserAssetDataSource/Private/AssetFileContextMenu.cpp b/Engine/Plugins/Editor/ContentBrowser/ContentBrowserAssetDataSource/Source/ContentBrowserAssetDataSource/Private/AssetFileContextMenu.cpp
--- a/Engine/Plugins/Editor/ContentBrowser/ContentBrowserAssetDataSource/Source/ContentBrowserAssetDataSource/Private/AssetFileContextMenu.cpp
+++ b/Engine/Plugins/Editor/ContentBrowser/ContentBrowserAssetDataSource/Source/ContentBrowserAssetDataSource/Private/AssetFileContextMenu.cpp
@@ -389,7 +389,7 @@ void FAssetFileContextMenu::MakeAssetActionsSubMenu(UToolMenu* Menu)
 		bool bHasExportableAssets = false;
 		for (const FAssetData& AssetData : SelectedAssets)
 		{
-			const UObject* Object = AssetData.GetAsset();
+			const UObject* Object = AssetData.FastGetAsset();
 			if (Object)
 			{
 				const UPackage* Package = Object->GetOutermost();
@@ -787,7 +787,7 @@ void FAssetFileContextMenu::MakeAssetLocalizationSubMenu(UToolMenu* Menu)
 			// Show the localization ID if we have a single asset selected
 			if (SelectedAssets.Num() == 1)
 			{
-				const FString LocalizationId = TextNamespaceUtil::GetPackageNamespace(SelectedAssets[0].GetAsset());
+				const FString LocalizationId = TextNamespaceUtil::GetPackageNamespace(SelectedAssets[0].FastGetAsset());
 				Section.AddMenuEntry(
 					"CopyLocalizationId",
 					FText::Format(LOCTEXT("CopyLocalizationIdFmt", "ID: {0}"), LocalizationId.IsEmpty() ? LOCTEXT("EmptyLocalizationId", "None") : FText::FromString(LocalizationId)),
@@ -1526,7 +1526,7 @@ void FAssetFileContextMenu::GetSelectedAssets(TArray<UObject*>& Assets, bool Ski
 			continue;
 		}
 
-		UObject* Object = SelectedAssets[AssetIdx].GetAsset();
+		UObject* Object = SelectedAssets[AssetIdx].FastGetAsset();
 
 		if (Object)
 		{
@@ -2209,10 +2209,12 @@ bool FAssetFileContextMenu::CanExecuteCreateBlueprintUsing() const
 	// Only work if you have a single asset selected
 	if(SelectedAssets.Num() == 1)
 	{
-		UObject* Asset = SelectedAssets[0].GetAsset();
-		// See if we know how to make a component from this asset
-		TArray< TSubclassOf<UActorComponent> > ComponentClassList = FComponentAssetBrokerage::GetComponentsForAsset(Asset);
-		return (ComponentClassList.Num() > 0);
+		if (UObject* Asset = SelectedAssets[0].FastGetAsset())
+		{
+			// See if we know how to make a component from this asset
+			TArray< TSubclassOf<UActorComponent> > ComponentClassList = FComponentAssetBrokerage::GetComponentsForAsset(Asset);
+			return (ComponentClassList.Num() > 0);
+		}
 	}
 
 	return false;
diff --git a/Engine/Plugins/Editor/ContentBrowser/ContentBrowserAssetDataSource/Source/ContentBrowserAssetDataSource/Private/ContentBrowserAssetDataCore.cpp b/Engine/Plugins/Editor/ContentBrowser/ContentBrowserAssetDataSource/Source/ContentBrowserAssetDataSource/Private/ContentBrowserAssetDataCore.cpp
--- a/Engine/Plugins/Editor/ContentBrowser/ContentBrowserAssetDataSource/Source/ContentBrowserAssetDataSource/Private/ContentBrowserAssetDataCore.cpp
+++ b/Engine/Plugins/Editor/ContentBrowser/ContentBrowserAssetDataSource/Source/ContentBrowserAssetDataSource/Private/ContentBrowserAssetDataCore.cpp
@@ -18,6 +18,7 @@
 #include "AssetFileContextMenu.h"
 #include "ContentBrowserDataUtils.h"
 #include "Settings/ContentBrowserSettings.h"
+#include "Engine/World.h"
 
 #define LOCTEXT_NAMESPACE "ContentBrowserAssetDataSource"
 
@@ -354,6 +355,7 @@ bool DuplicateItem(IAssetTools* InAssetTools, const UContentBrowserDataSource* I
 bool DuplicateAssetFileItem(IAssetTools* InAssetTools, const FContentBrowserAssetFileItemDataPayload& InAssetPayload, UObject*& OutSourceAsset, FAssetData& OutNewAsset)
 {
 	// We need to potentially load the asset in order to duplicate it
+	FScopedLoadAllExternalObjects Scope(InAssetPayload.GetAssetData().PackageName);
 	if (UObject* Asset = InAssetPayload.LoadAsset())
 	{
 		// Find a unique default name for the duplicated asset
@@ -392,6 +394,7 @@ bool DuplicateAssetFileItems(TArrayView<const TSharedRef<const FContentBrowserAs
 	for (const TSharedRef<const FContentBrowserAssetFileItemDataPayload>& AssetPayload : InAssetPayloads)
 	{
 		// We need to potentially load the asset in order to duplicate it
+		FScopedLoadAllExternalObjects Scope(AssetPayload->GetAssetData().PackageName);
 		if (UObject* Asset = AssetPayload->LoadAsset())
 		{
 			ObjectsToDuplicate.Add(Asset);
@@ -767,6 +770,7 @@ bool RenameAssetFolderItem(IAssetRegistry* InAssetRegistry, const FContentBrowse
 bool RenameAssetFileItem(IAssetTools* InAssetTools, const FContentBrowserAssetFileItemDataPayload& InAssetPayload, const FString& InNewName)
 {
 	// We need to potentially load the asset in order to rename it
+	FScopedLoadAllExternalObjects Scope(InAssetPayload.GetAssetData().PackageName);
 	if (UObject* Asset = InAssetPayload.LoadAsset())
 	{
 		const FString PackagePath = FPackageName::GetLongPackagePath(Asset->GetOutermost()->GetName());
@@ -840,6 +844,7 @@ bool CopyAssetFileItems(TArrayView<const TSharedRef<const FContentBrowserAssetFi
 	for (const TSharedRef<const FContentBrowserAssetFileItemDataPayload>& AssetPayload : InAssetPayloads)
 	{
 		// We need to potentially load the asset in order to duplicate it
+		FScopedLoadAllExternalObjects Scope(AssetPayload->GetAssetData().PackageName);
 		if (UObject* Asset = AssetPayload->LoadAsset())
 		{
 			AssetsToCopy.Add(Asset);
@@ -925,6 +930,7 @@ bool MoveAssetFileItems(TArrayView<const TSharedRef<const FContentBrowserAssetFi
 	for (const TSharedRef<const FContentBrowserAssetFileItemDataPayload>& AssetPayload : InAssetPayloads)
 	{
 		// We need to potentially load the asset in order to duplicate it
+		FScopedLoadAllExternalObjects Scope(AssetPayload->GetAssetData().PackageName);
 		if (UObject* Asset = AssetPayload->LoadAsset())
 		{
 			AssetsToMove.Add(Asset);
diff --git a/Engine/Source/Developer/AssetTools/Private/AssetTypeActions/AssetTypeActions_World.cpp b/Engine/Source/Developer/AssetTools/Private/AssetTypeActions/AssetTypeActions_World.cpp
--- a/Engine/Source/Developer/AssetTools/Private/AssetTypeActions/AssetTypeActions_World.cpp
+++ b/Engine/Source/Developer/AssetTools/Private/AssetTypeActions/AssetTypeActions_World.cpp
@@ -7,6 +7,29 @@
 
 #define LOCTEXT_NAMESPACE "AssetTypeActions"
 
+namespace AssetTypeActions_World
+{
+	bool CanRenameOrDuplicate(const FAssetData& InAsset, FText* OutErrorMsg)
+	{
+		if (ULevel::GetIsLevelPartitionedFromAsset(InAsset))
+		{
+			for (const FWorldContext& WorldContext : GEditor->GetWorldContexts())
+			{
+				if (const UWorld* World = WorldContext.World(); WorldContext.World() && InAsset.PackageName == World->GetPackage()->GetFName())
+				{
+					if (OutErrorMsg)
+					{
+						*OutErrorMsg = LOCTEXT("Error_CannotDuplicatRenameWorldPartitionWhileInUse", "Cannot duplicate / rename a partition world while it is used.");
+					}
+					return false;
+				}
+			}
+		}
+
+		return true;
+	}
+}
+
 void FAssetTypeActions_World::OpenAssetEditor( const TArray<UObject*>& InObjects, TSharedPtr<class IToolkitHost> EditWithinLevelEditor )
 {
 	for (auto ObjIt = InObjects.CreateConstIterator(); ObjIt; ++ObjIt)
@@ -59,4 +82,14 @@ TArray<FAssetData> FAssetTypeActions_World::GetValidAssetsForPreviewOrEdit(TArra
 	return MoveTemp(AssetsToOpen);
 }
 
+bool FAssetTypeActions_World::CanRename(const FAssetData& InAsset, FText* OutErrorMsg) const
+{
+	return AssetTypeActions_World::CanRenameOrDuplicate(InAsset, OutErrorMsg);
+}
+
+bool FAssetTypeActions_World::CanDuplicate(const FAssetData& InAsset, FText* OutErrorMsg) const
+{
+	return AssetTypeActions_World::CanRenameOrDuplicate(InAsset, OutErrorMsg);
+}
+
 #undef LOCTEXT_NAMESPACE
diff --git a/Engine/Source/Developer/AssetTools/Private/AssetTypeActions/AssetTypeActions_World.h b/Engine/Source/Developer/AssetTools/Private/AssetTypeActions/AssetTypeActions_World.h
--- a/Engine/Source/Developer/AssetTools/Private/AssetTypeActions/AssetTypeActions_World.h
+++ b/Engine/Source/Developer/AssetTools/Private/AssetTypeActions/AssetTypeActions_World.h
@@ -20,4 +20,6 @@ public:
 	virtual bool CanLocalize() const override { return false; }
 	virtual class UThumbnailInfo* GetThumbnailInfo(UObject* Asset) const override;
 	virtual TArray<FAssetData> GetValidAssetsForPreviewOrEdit(TArrayView<const FAssetData> InAssetDatas, bool bIsPreview) override;
+	virtual bool CanRename(const FAssetData& InAsset, FText* OutErrorMsg) const override;
+	virtual bool CanDuplicate(const FAssetData& InAsset, FText* OutErrorMsg) const override;
 };
diff --git a/Engine/Source/Developer/AssetTools/Private/AssetViewUtils.cpp b/Engine/Source/Developer/AssetTools/Private/AssetViewUtils.cpp
--- a/Engine/Source/Developer/AssetTools/Private/AssetViewUtils.cpp
+++ b/Engine/Source/Developer/AssetTools/Private/AssetViewUtils.cpp
@@ -15,6 +15,7 @@
 #include "Modules/ModuleManager.h"
 #include "UnrealClient.h"
 #include "Engine/World.h"
+#include "Engine/Level.h"
 #include "Settings/ContentBrowserSettings.h"
 #include "Settings/EditorExperimentalSettings.h"
 #include "SourceControlOperations.h"
@@ -57,6 +58,12 @@ namespace AssetViewUtils
 
 	/** Internal function to delete a folder from disk, but only if it is empty. InPathToDelete is in FPackageName format. */
 	bool DeleteEmptyFolderFromDisk(const FString& InPathToDelete);
+
+	/** Get all the objects in a list of asset data with optional load of all external packages */
+	void GetObjectsInAssetData(const TArray<FAssetData>& AssetList, TArray<UObject*>& OutDroppedObjects, bool bLoadAllExternalObjects);
+
+	/** Makes sure the specified assets are loaded into memory. */
+	bool LoadAssetsIfNeeded(const TArray<FString>& ObjectPaths, TArray<UObject*>& LoadedObjects, bool bAllowedToPromptToLoadAssets, bool bLoadRedirects, bool bLoadWorldPartitionWorlds, bool bLoadAllExternalObjects);
 }
 
 bool AssetViewUtils::OpenEditorForAsset(const FString& ObjectPath)
@@ -65,7 +72,13 @@ bool AssetViewUtils::OpenEditorForAsset(const FString& ObjectPath)
 	TArray<UObject*> LoadedObjects;
 	TArray<FString> ObjectPaths;
 	ObjectPaths.Add(ObjectPath);
-	LoadAssetsIfNeeded(ObjectPaths, LoadedObjects);
+
+	// Here we want to load the asset as it will be passed to OpenEditorForAsset
+	const bool bAllowedToPromptToLoadAssets = true;
+	const bool bLoadRedirects = false;
+	const bool bLoadWorldPartitionWorlds = true;
+	const bool bLoadAllExternalObjects = false;
+	LoadAssetsIfNeeded(ObjectPaths, LoadedObjects, bAllowedToPromptToLoadAssets, bLoadRedirects, bLoadWorldPartitionWorlds, bLoadAllExternalObjects);
 
 	// Open the editor for the specified asset
 	UObject* FoundObject = FindObject<UObject>(NULL, *ObjectPath);
@@ -99,6 +112,13 @@ bool AssetViewUtils::OpenEditorForAsset(const TArray<UObject*>& Assets)
 }
 
 bool AssetViewUtils::LoadAssetsIfNeeded(const TArray<FString>& ObjectPaths, TArray<UObject*>& LoadedObjects, bool bAllowedToPromptToLoadAssets, bool bLoadRedirects)
+{
+	const bool bLoadWorldPartitionWorlds = false;
+	const bool bLoadAllExternalObjects = false;
+	return LoadAssetsIfNeeded(ObjectPaths, LoadedObjects, bAllowedToPromptToLoadAssets, bLoadRedirects, bLoadWorldPartitionWorlds, bLoadAllExternalObjects);
+}
+
+bool AssetViewUtils::LoadAssetsIfNeeded(const TArray<FString>& ObjectPaths, TArray<UObject*>& LoadedObjects, bool bAllowedToPromptToLoadAssets, bool bLoadRedirects, bool bLoadWorldPartitionWorlds, bool bLoadAllExternalObjects)
 {
 	bool bAnyObjectsWereLoadedOrUpdated = false;
 
@@ -116,12 +136,19 @@ bool AssetViewUtils::LoadAssetsIfNeeded(const TArray<FString>& ObjectPaths, TArr
 		}
 		else
 		{
-			// Unloaded asset, we will load it later
-			UnloadedObjectPaths.Add(ObjectPath);
 			if ( FEditorFileUtils::IsMapPackageAsset(ObjectPath) )
 			{
+				FName PackageName = FName(*FEditorFileUtils::ExtractPackageName(ObjectPath));
+				if (!bLoadWorldPartitionWorlds && ULevel::GetIsLevelPartitionedFromPackage(PackageName))
+				{
+					continue;
+				}
+				
 				bAtLeastOneUnloadedMap = true;
 			}
+
+			// Unloaded asset, we will load it later
+			UnloadedObjectPaths.Add(ObjectPath);
 		}
 	}
 
@@ -150,6 +177,7 @@ bool AssetViewUtils::LoadAssetsIfNeeded(const TArray<FString>& ObjectPaths, TArr
 			SlowTask.EnterProgressFrame(1, FText::Format(LOCTEXT("LoadingObjectf", "Loading {0}..."), FText::FromString(ObjectPath)));
 
 			// Load up the object
+			TUniquePtr<FScopedLoadAllExternalObjects> Scope(bLoadAllExternalObjects? new FScopedLoadAllExternalObjects(*FEditorFileUtils::ExtractPackageName(ObjectPath)) : nullptr);
 			UObject* LoadedObject = LoadObject<UObject>(NULL, *ObjectPath, NULL, LoadFlags, NULL);
 			if ( LoadedObject )
 			{
@@ -339,7 +367,11 @@ bool AssetViewUtils::DeleteFolders(const TArray<FString>& PathsToDelete)
 
 		// Load all the assets in the selected paths
 		TArray<UObject*> LoadedAssets;
-		if ( LoadAssetsIfNeeded(ObjectPaths, LoadedAssets) )
+		const bool bAllowedToPromptToLoadAssets = true;
+		const bool bLoadRedirects = false;
+		const bool bLoadWorldPartitionWorlds = true;
+		const bool bLoadAllExternalObjects = true;
+		if ( LoadAssetsIfNeeded(ObjectPaths, LoadedAssets, bAllowedToPromptToLoadAssets, bLoadRedirects, bLoadWorldPartitionWorlds, bLoadAllExternalObjects))
 		{
 			// Make sure we loaded all of them
 			if ( LoadedAssets.Num() == NumAssetsInPaths )
@@ -515,7 +547,8 @@ bool AssetViewUtils::RenameFolder(const FString& DestPath, const FString& Source
 	TArray<FAssetData> AssetsInFolder;
 	AssetRegistryModule.Get().GetAssetsByPath(*SourcePath, AssetsInFolder, true);
 	TArray<UObject*> ObjectsInFolder;
-	GetObjectsInAssetData(AssetsInFolder, ObjectsInFolder);
+	const bool bLoadAllExternalObjects = true;
+	GetObjectsInAssetData(AssetsInFolder, ObjectsInFolder, bLoadAllExternalObjects);
 	MoveAssets(ObjectsInFolder, DestPath, SourcePath);
 
 	// Now check to see if the original folder is empty, if so we can delete it
@@ -700,7 +733,11 @@ bool AssetViewUtils::PrepareFoldersForDragDrop(const TArray<FString>& SourcePath
 
 		// Load all assets in this path if needed
 		TArray<UObject*> AllLoadedAssets;
-		LoadAssetsIfNeeded(ObjectPaths, AllLoadedAssets, false);
+		const bool bAllowedToPromptToLoadAssets = false;
+		const bool bLoadRedirects = false;
+		const bool bLoadWorldPartitionWorlds = true;
+		const bool bLoadAllExternalObjects = true;
+		LoadAssetsIfNeeded(ObjectPaths, AllLoadedAssets, bAllowedToPromptToLoadAssets, bLoadRedirects, bLoadWorldPartitionWorlds, bLoadAllExternalObjects);
 
 		// Add a slash to the end of the path so StartsWith doesn't get a false positive on similarly named folders
 		const FString SourcePathWithSlash = *PathIt + TEXT("/");
@@ -949,12 +986,12 @@ bool AssetViewUtils::IsPluginFolder(const FStringView InPath, EPluginLoadedFrom*
 	return false;
 }
 
-void AssetViewUtils::GetObjectsInAssetData(const TArray<FAssetData>& AssetList, TArray<UObject*>& OutDroppedObjects)
-{	
+void AssetViewUtils::GetObjectsInAssetData(const TArray<FAssetData>& AssetList, TArray<UObject*>& OutDroppedObjects, bool bLoadAllExternalObjects)
+{
 	for (int32 AssetIdx = 0; AssetIdx < AssetList.Num(); ++AssetIdx)
 	{
 		const FAssetData& AssetData = AssetList[AssetIdx];
-
+		TUniquePtr<FScopedLoadAllExternalObjects> LoadAllExternalObjects(bLoadAllExternalObjects ? new FScopedLoadAllExternalObjects(AssetData.PackageName) : nullptr);
 		UObject* Obj = AssetData.GetAsset();
 		if (Obj)
 		{
@@ -963,6 +1000,12 @@ void AssetViewUtils::GetObjectsInAssetData(const TArray<FAssetData>& AssetList,
 	}
 }
 
+void AssetViewUtils::GetObjectsInAssetData(const TArray<FAssetData>& AssetList, TArray<UObject*>& OutDroppedObjects)
+{	
+	const bool bLoadAllExternalObjects = false;
+	GetObjectsInAssetData(AssetList, OutDroppedObjects, bLoadAllExternalObjects);
+}
+
 bool AssetViewUtils::IsValidFolderName(const FString& FolderName, FText& Reason)
 {
 	// Check length of the folder name
diff --git a/Engine/Source/Editor/ContentBrowser/Private/SAssetView.cpp b/Engine/Source/Editor/ContentBrowser/Private/SAssetView.cpp
--- a/Engine/Source/Editor/ContentBrowser/Private/SAssetView.cpp
+++ b/Engine/Source/Editor/ContentBrowser/Private/SAssetView.cpp
@@ -1522,23 +1522,16 @@ FReply SAssetView::OnKeyChar( const FGeometry& MyGeometry,const FCharacterEvent&
 	return FReply::Unhandled();
 }
 
-static bool IsValidObjectPath(const FString& Path)
+static bool IsValidObjectPath(const FString& Path, FString& OutObjectClassName, FString& OutObjectPath, FString& OutPackageName)
 {
-	int32 NameStartIndex = INDEX_NONE;
-	Path.FindChar(TCHAR('\''), NameStartIndex);
-	if (NameStartIndex != INDEX_NONE)
+	if (FPackageName::ParseExportTextPath(Path, &OutObjectClassName, &OutObjectPath))
 	{
-		int32 NameEndIndex = INDEX_NONE;
-		Path.FindLastChar(TCHAR('\''), NameEndIndex);
-		if (NameEndIndex > NameStartIndex)
+		if (UClass* ObjectClass = FindObject<UClass>(ANY_PACKAGE, *OutObjectClassName, true))
 		{
-			const FString ClassName = Path.Left(NameStartIndex);
-			const FString PathName = Path.Mid(NameStartIndex + 1, NameEndIndex - NameStartIndex - 1);
-
-			UClass* Class = FindObject<UClass>(ANY_PACKAGE, *ClassName, true);
-			if (Class)
+			OutPackageName = FPackageName::ObjectPathToPackageName(OutObjectPath);
+			if (FPackageName::IsValidLongPackageName(OutPackageName))
 			{
-				return FPackageName::IsValidLongPackageName(FPackageName::ObjectPathToPackageName(PathName));
+				return true;
 			}
 		}
 	}
@@ -1575,9 +1568,13 @@ FReply SAssetView::OnKeyDown( const FGeometry& MyGeometry, const FKeyEvent& InKe
 			for (const FString& AssetPath : AssetPathsSplit)
 			{
 				// Validate string
-				if (IsValidObjectPath(AssetPath))
+				FString ObjectClassName;
+				FString ObjectPath;
+				FString PackageName;
+				if (IsValidObjectPath(AssetPath, ObjectClassName, ObjectPath, PackageName))
 				{
-					UObject* ObjectToCopy = LoadObject<UObject>(nullptr, *AssetPath);
+					FScopedLoadAllExternalObjects Scope(*PackageName);
+					UObject* ObjectToCopy = LoadObject<UObject>(nullptr, *ObjectPath);
 					if (ObjectToCopy && !ObjectToCopy->IsA(UClass::StaticClass()))
 					{
 						AssetsToCopy.Add(ObjectToCopy);
diff --git a/Engine/Source/Editor/UnrealEd/Private/EditorEngine.cpp b/Engine/Source/Editor/UnrealEd/Private/EditorEngine.cpp
--- a/Engine/Source/Editor/UnrealEd/Private/EditorEngine.cpp
+++ b/Engine/Source/Editor/UnrealEd/Private/EditorEngine.cpp
@@ -4232,6 +4232,18 @@ bool UEditorEngine::InitializePhysicsSceneForSaveIfNecessary(UWorld* World, bool
 
 void UEditorEngine::CleanupPhysicsSceneThatWasInitializedForSave(UWorld* World, bool bForceInitialized)
 {
+	// Capture Dirty packages so we can reset unwanted dirtyness 
+	TSet<UPackage*> DirtyPackages;
+	const bool bWorldPackageDirty = World->GetPackage()->IsDirty();
+
+	for (UPackage* ExternalPackage : World->GetPackage()->GetExternalPackages())
+	{
+		if (ExternalPackage->IsDirty())
+		{
+			DirtyPackages.Add(ExternalPackage);
+		}
+	}
+
 	// Make sure we clean up the physics scene here. If we leave too many scenes in memory, undefined behavior occurs when locking a scene for read/write.
 	World->ClearWorldComponents();
 
@@ -4252,6 +4264,19 @@ void UEditorEngine::CleanupPhysicsSceneThatWasInitializedForSave(UWorld* World,
 
 	// Update components again in case it was a world without a physics scene but did have rendered components.
 	World->UpdateWorldComponents(true, true);
+
+	for (UPackage* ExternalPackage : World->GetPackage()->GetExternalPackages())
+	{
+		if (!DirtyPackages.Contains(ExternalPackage) && ExternalPackage->IsDirty())
+		{
+			ExternalPackage->SetDirtyFlag(false);
+		}
+	}
+
+	if (!bWorldPackageDirty && World->GetPackage()->IsDirty())
+	{
+		World->GetPackage()->SetDirtyFlag(false);
+	}
 }
 
 FSavePackageResultStruct UEditorEngine::Save(UPackage* InOuter, UObject* InBase, EObjectFlags TopLevelFlags,
diff --git a/Engine/Source/Editor/UnrealEd/Private/FileHelpers.cpp b/Engine/Source/Editor/UnrealEd/Private/FileHelpers.cpp
--- a/Engine/Source/Editor/UnrealEd/Private/FileHelpers.cpp
+++ b/Engine/Source/Editor/UnrealEd/Private/FileHelpers.cpp
@@ -744,6 +744,9 @@ static bool SaveWorld(UWorld* World,
 		// Save Loaded cells
 		TArray<FName> LoadedEditorCells;
 
+		// Initialize Physics Scene for save if needed here so that external packages don't get dirtied during the Saving of the map package
+		bool bForceInitializedWorld = false;
+		
 		if ( bRenamePackageToFile )
 		{
 			if (bPackageNeedsRename)
@@ -827,6 +830,11 @@ static bool SaveWorld(UWorld* World,
 		SlowTask.EnterProgressFrame(25);
 				
 		UWorld* SaveWorld = DuplicatedWorld ? DuplicatedWorld : World;
+		const bool bNewlyCreated = SaveWorld->GetPackage()->HasAnyPackageFlags(PKG_NewlyCreated);
+
+		// Initialize Physics Scene for save if needed here before saving external packages as this can modify those external package objects
+		// This makes UEditorEngine::Save's own call to InitializePhysicsSceneForSaveIfNecessary redundant but wasn't removed to avoid breaking other code paths
+		const bool bInitializedPhysicsSceneForSave = GEditor->InitializePhysicsSceneForSaveIfNecessary(SaveWorld, bForceInitializedWorld);
 				
 		if(!bAutosaving)
 		{
@@ -849,7 +857,7 @@ static bool SaveWorld(UWorld* World,
 			bSuccess = GEditor->Exec(NULL, *FString::Printf(TEXT("OBJ SAVEPACKAGE PACKAGE=\"%s\" FILE=\"%s\" SILENT=true AUTOSAVING=%s KEEPDIRTY=%s"), *Package->GetName(), *FinalFilename, *AutoSavingString, *KeepDirtyString), SaveErrors);
 			SaveErrors.Flush();
 		}
-				
+		
 		if (bSuccess)
 		{
 			// Force update before initializing World Partition
@@ -859,7 +867,7 @@ static bool SaveWorld(UWorld* World,
 			// Make sure when we exit SaveWorld AssetRegistry is up to date with saved map
 			AssetRegistry.ScanModifiedAssetFiles({ FinalFilename });
 			
-			if (bPackageNeedsRename)
+			if (bPackageNeedsRename || bNewlyCreated)
 			{
 				// Force rescan to make sure assets are found on map open or world partition initialize
 				AssetRegistry.ScanPathsSynchronous( ULevel::GetExternalObjectsPaths(NewPackageName) , true);
@@ -889,6 +897,13 @@ static bool SaveWorld(UWorld* World,
 			GEditor->RunDeferredMarkForAddFiles();
 		}
 
+
+		// If Physics scene was initialized for save, cleanup
+		if (bInitializedPhysicsSceneForSave)
+		{
+			GEditor->CleanupPhysicsSceneThatWasInitializedForSave(SaveWorld, bForceInitializedWorld);
+		}
+
 		// If the package save was not successful. Trash the duplicated world or rename back if the duplicate failed.
 		if( bRenamePackageToFile && !bSuccess )
 		{
diff --git a/Engine/Source/Editor/UnrealEd/Private/ObjectTools.cpp b/Engine/Source/Editor/UnrealEd/Private/ObjectTools.cpp
--- a/Engine/Source/Editor/UnrealEd/Private/ObjectTools.cpp
+++ b/Engine/Source/Editor/UnrealEd/Private/ObjectTools.cpp
@@ -2,6 +2,7 @@
 
 
 #include "ObjectTools.h"
+#include "Engine/World.h"
 #include "Engine/Level.h"
 #include "UObject/UnrealType.h"
 #include "Components/ActorComponent.h"
@@ -2290,6 +2291,7 @@ namespace ObjectTools
 		for ( int i = 0; i < AssetsToDelete.Num(); i++ )
 		{
 			const FAssetData& AssetData = AssetsToDelete[i];
+			FScopedLoadAllExternalObjects Scope(AssetData.PackageName);
 			UObject *ObjectToDelete = AssetData.GetAsset();
 			// Assets can be loaded even when their underlying type/class no longer exists...
 			if ( ObjectToDelete!=nullptr )
diff --git a/Engine/Source/Runtime/Engine/Classes/Engine/Level.h b/Engine/Source/Runtime/Engine/Classes/Engine/Level.h
--- a/Engine/Source/Runtime/Engine/Classes/Engine/Level.h
+++ b/Engine/Source/Runtime/Engine/Classes/Engine/Level.h
@@ -764,6 +764,12 @@ public:
 	ENGINE_API void MarkLevelBoundsDirty();
 
 #if WITH_EDITOR
+	ENGINE_API static bool GetLevelBoundsFromAsset(const FAssetData& Asset, FBox& OutLevelBounds);
+	ENGINE_API static bool GetIsLevelPartitionedFromAsset(const FAssetData& Asset);
+	ENGINE_API static bool GetIsLevelUsingExternalActorsFromAsset(const FAssetData& Asset);
+	ENGINE_API static bool GetIsUsingActorFoldersFromAsset(const FAssetData& Asset);
+
+
 	ENGINE_API static bool GetLevelBoundsFromPackage(FName LevelPackage, FBox& OutLevelBounds);
 	ENGINE_API static bool GetIsLevelPartitionedFromPackage(FName LevelPackage);
 	ENGINE_API static bool GetIsLevelUsingExternalActorsFromPackage(FName LevelPackage);
diff --git a/Engine/Source/Runtime/Engine/Classes/Engine/World.h b/Engine/Source/Runtime/Engine/Classes/Engine/World.h
--- a/Engine/Source/Runtime/Engine/Classes/Engine/World.h
+++ b/Engine/Source/Runtime/Engine/Classes/Engine/World.h
@@ -984,6 +984,17 @@ public:
 	void Reevaluate(ULevelStreaming* StreamingLevel);
 };
 
+/**
+ * Scope can be used to force load all external objects when loading a world. 
+ */
+struct ENGINE_API FScopedLoadAllExternalObjects
+{
+	FScopedLoadAllExternalObjects(FName InPackageName);
+	~FScopedLoadAllExternalObjects();
+
+	FName PackageName;
+};
+
 /** 
  * The World is the top level object representing a map or a sandbox in which Actors and Components will exist and be rendered.  
  *
@@ -1475,7 +1486,11 @@ public:
 	/** Returns whether or not this world is currently ticking. See SetShouldTick. */
 	bool ShouldTick() const { return bShouldTick; }
 
+	static bool ShouldLoadAllExternalObjects(FName InPackageName) { return LoadAllExternalObjects.Contains(InPackageName); }
+
 private:
+	friend struct FScopedLoadAllExternalObjects;
+	static TSet<FName> LoadAllExternalObjects;
 
 	/** List of all the controllers in the world. */
 	TArray<TWeakObjectPtr<class AController> > ControllerList;
diff --git a/Engine/Source/Runtime/Engine/Private/Level.cpp b/Engine/Source/Runtime/Engine/Private/Level.cpp
--- a/Engine/Source/Runtime/Engine/Private/Level.cpp
+++ b/Engine/Source/Runtime/Engine/Private/Level.cpp
@@ -888,7 +888,7 @@ void ULevel::PostLoad()
 	// if we use external actors, load dynamic actors here
 	if (IsUsingExternalActors())
 	{
-		if (!bWasDuplicated && !bIsPartitioned)
+		if (!bWasDuplicated && (!bIsPartitioned || UWorld::ShouldLoadAllExternalObjects(GetPackage()->GetFName())))
 		{
 			UPackage* LevelPackage = GetPackage();
 			bool bPackageForPIE = LevelPackage->HasAnyPackageFlags(PKG_PlayInEditor);
@@ -2069,75 +2069,95 @@ namespace LevelAssetRegistryHelper
 	}
 }
 
+bool ULevel::GetIsLevelPartitionedFromAsset(const FAssetData& Asset)
+{
+	FString LevelIsPartitionedStr;
+	static const FName NAME_LevelIsPartitioned(TEXT("LevelIsPartitioned"));
+	if (Asset.GetTagValue(NAME_LevelIsPartitioned, LevelIsPartitionedStr))
+	{
+		check(LevelIsPartitionedStr == TEXT("1"));
+		return true;
+	}
+	return false;
+}
+
 bool ULevel::GetIsLevelPartitionedFromPackage(FName LevelPackage)
 {
 	return LevelPackage.IsNone() ? false : LevelAssetRegistryHelper::GetLevelInfoFromAssetRegistry(LevelPackage, [](const FAssetData& Asset)
 	{
-		FString LevelIsPartitionedStr;
-		static const FName NAME_LevelIsPartitioned(TEXT("LevelIsPartitioned"));
-		if (Asset.GetTagValue(NAME_LevelIsPartitioned, LevelIsPartitionedStr))
-		{
-			check(LevelIsPartitionedStr == TEXT("1"));
-			return true;
-		}
-		return false;
+		return GetIsLevelPartitionedFromAsset(Asset);
 	});
 }
 
+bool ULevel::GetIsLevelUsingExternalActorsFromAsset(const FAssetData& Asset)
+{
+	FString LevelIsUsingExternalActorsStr;
+	static const FName NAME_LevelIsUsingExternalActors(TEXT("LevelIsUsingExternalActors"));
+	if (Asset.GetTagValue(NAME_LevelIsUsingExternalActors, LevelIsUsingExternalActorsStr))
+	{
+		check(LevelIsUsingExternalActorsStr == TEXT("1"));
+		return true;
+	}
+	return false;
+}
+
 bool ULevel::GetIsLevelUsingExternalActorsFromPackage(FName LevelPackage)
 {
 	return LevelAssetRegistryHelper::GetLevelInfoFromAssetRegistry(LevelPackage, [](const FAssetData& Asset)
 	{
-		FString LevelIsUsingExternalActorsStr;
-		static const FName NAME_LevelIsUsingExternalActors(TEXT("LevelIsUsingExternalActors"));
-		if (Asset.GetTagValue(NAME_LevelIsUsingExternalActors, LevelIsUsingExternalActorsStr))
-		{
-			check(LevelIsUsingExternalActorsStr == TEXT("1"));
-			return true;
-		}
-		return false;
+		return GetIsLevelUsingExternalActorsFromAsset(Asset);
 	});
 }
 
+bool ULevel::GetIsUsingActorFoldersFromAsset(const FAssetData& Asset)
+{
+	FString LevelIsUsingActorFoldersStr;
+	static const FName NAME_LevelIsUsingActorFolders(TEXT("LevelIsUsingActorFolders"));
+	if (Asset.GetTagValue(NAME_LevelIsUsingActorFolders, LevelIsUsingActorFoldersStr))
+	{
+		check(LevelIsUsingActorFoldersStr == TEXT("1"));
+		return true;
+	}
+	return false;
+}
+
 bool ULevel::GetIsUsingActorFoldersFromPackage(FName LevelPackage)
 {
 	return LevelAssetRegistryHelper::GetLevelInfoFromAssetRegistry(LevelPackage, [](const FAssetData& Asset)
 	{
-		FString LevelIsUsingActorFoldersStr;
-		static const FName NAME_LevelIsUsingActorFolders(TEXT("LevelIsUsingActorFolders"));
-		if (Asset.GetTagValue(NAME_LevelIsUsingActorFolders, LevelIsUsingActorFoldersStr))
+		return GetIsUsingActorFoldersFromAsset(Asset);
+	});
+}
+
+bool ULevel::GetLevelBoundsFromAsset(const FAssetData& Asset, FBox& OutLevelBounds)
+{
+	FString LevelBoundsLocationStr;
+	static const FName NAME_LevelBoundsLocation(TEXT("LevelBoundsLocation"));
+
+	FString LevelBoundsExtentStr;
+	static const FName NAME_LevelBoundsExtent(TEXT("LevelBoundsExtent"));
+
+	if (Asset.GetTagValue(NAME_LevelBoundsLocation, LevelBoundsLocationStr) &&
+		Asset.GetTagValue(NAME_LevelBoundsExtent, LevelBoundsExtentStr))
+	{
+		FVector LevelBoundsLocation;
+		FVector LevelBoundsExtent;
+		if (LevelBoundsLocation.InitFromCompactString(LevelBoundsLocationStr) &&
+			LevelBoundsExtent.InitFromCompactString(LevelBoundsExtentStr))
 		{
-			check(LevelIsUsingActorFoldersStr == TEXT("1"));
+			OutLevelBounds = FBox(LevelBoundsLocation - LevelBoundsExtent, LevelBoundsLocation + LevelBoundsExtent);
 			return true;
 		}
-		return false;
-	});
+	}
+	// Invalid bounds
+	return false;
 }
 
 bool ULevel::GetLevelBoundsFromPackage(FName LevelPackage, FBox& OutLevelBounds)
 {
 	return LevelAssetRegistryHelper::GetLevelInfoFromAssetRegistry(LevelPackage, [&OutLevelBounds](const FAssetData& Asset)
 	{
-		FString LevelBoundsLocationStr;
-		static const FName NAME_LevelBoundsLocation(TEXT("LevelBoundsLocation"));
-
-		FString LevelBoundsExtentStr;
-		static const FName NAME_LevelBoundsExtent(TEXT("LevelBoundsExtent"));
-
-		if (Asset.GetTagValue(NAME_LevelBoundsLocation, LevelBoundsLocationStr) &&
-			Asset.GetTagValue(NAME_LevelBoundsExtent, LevelBoundsExtentStr))
-		{
-			FVector LevelBoundsLocation;
-			FVector LevelBoundsExtent;
-			if (LevelBoundsLocation.InitFromCompactString(LevelBoundsLocationStr) &&
-				LevelBoundsExtent.InitFromCompactString(LevelBoundsExtentStr))
-			{
-				OutLevelBounds = FBox(LevelBoundsLocation - LevelBoundsExtent, LevelBoundsLocation + LevelBoundsExtent);
-				return true;
-			}
-		}
-		// Invalid bounds
-		return false;
+		return GetLevelBoundsFromAsset(Asset, OutLevelBounds);
 	});
 }
 
diff --git a/Engine/Source/Runtime/Engine/Private/World.cpp b/Engine/Source/Runtime/Engine/Private/World.cpp
--- a/Engine/Source/Runtime/Engine/Private/World.cpp
+++ b/Engine/Source/Runtime/Engine/Private/World.cpp
@@ -319,6 +319,24 @@ FScopedLevelCollectionContextSwitch::~FScopedLevelCollectionContextSwitch()
 	}
 }
 
+FScopedLoadAllExternalObjects::FScopedLoadAllExternalObjects(FName InPackageName)
+{
+	if (!UWorld::LoadAllExternalObjects.Contains(InPackageName))
+	{
+		PackageName = InPackageName;
+		UWorld::LoadAllExternalObjects.Add(PackageName);
+	}
+}
+
+FScopedLoadAllExternalObjects::~FScopedLoadAllExternalObjects()
+{
+	if (PackageName != NAME_None)
+	{
+		UWorld::LoadAllExternalObjects.Remove(PackageName);
+	}
+}
+
+
 FAudioDeviceWorldDelegates::FOnWorldRegisteredToAudioDevice FAudioDeviceWorldDelegates::OnWorldRegisteredToAudioDevice;
 
 FAudioDeviceWorldDelegates::FOnWorldUnregisteredWithAudioDevice FAudioDeviceWorldDelegates::OnWorldUnregisteredWithAudioDevice;
@@ -331,6 +349,7 @@ FAudioDeviceWorldDelegates::FOnWorldUnregisteredWithAudioDevice FAudioDeviceWorl
 UWorldProxy GWorld;
 
 TMap<FName, EWorldType::Type> UWorld::WorldTypePreLoadMap;
+TSet<FName> UWorld::LoadAllExternalObjects;
 
 FWorldDelegates::FWorldEvent FWorldDelegates::OnPostWorldCreation;
 FWorldDelegates::FWorldInitializationEvent FWorldDelegates::OnPreWorldInitialization;
