Support for migrating selected assets and their dependencies that are inter-referential across content roots.

Assets spread across content roots are migrated into the corresponding content root in the destination, if it exists. If any don't exist, migrate the selection and dependencies to temporary packages into a folder that is migrated to the destination.

== Notes ==

AssetTools.h/cpp
* Remove OriginalRoot filtering down reference gathering to objects in the selected object package's content root (opens up the routine to migrating plugin content)
* Trawl the set of objects to migrate and compare their content root against the available content roots in the destination project's plugin directory hierarchy
* If some plugins are missing from the destination, prompt the user to consolidate everything into a temporary folder
* If the user selects to consolidate we duplicate objects and fixup references to those duplicates, and save those packages to a folder named by the user
* If necessary copy those temporary packages to the destination content root
* Tidy up any temporary packages saved into the source project

diff --git a/Engine/Source/Developer/AssetTools/Private/AssetTools.cpp b/Engine/Source/Developer/AssetTools/Private/AssetTools.cpp
--- a/Engine/Source/Developer/AssetTools/Private/AssetTools.cpp
+++ b/Engine/Source/Developer/AssetTools/Private/AssetTools.cpp
@@ -151,6 +151,10 @@
 #include "InterchangeProjectSettings.h"
 #include "Engine/World.h"
 #include "Engine/Level.h"
+#include "UObject/SavePackage.h"
+#include "Dialogs/Dialogs.h"
+#include "Widgets/Input/SEditableTextBox.h"
+#include "Widgets/Layout/SSpacer.h"
 
 #if WITH_EDITOR
 #include "Subsystems/AssetEditorSubsystem.h"
@@ -2811,18 +2815,13 @@ void UAssetToolsImpl::PerformMigratePackages(TArray<FName> PackageNamesToMigrate
 			if ( !AllPackageNamesToMove.Contains(*PackageIt) )
 			{
 				AllPackageNamesToMove.Add(*PackageIt);
-				FString Path = (*PackageIt).ToString();
-				FString OriginalRootString;
-				Path.RemoveFromStart(TEXT("/"));
-				Path.Split("/", &OriginalRootString, &Path, ESearchCase::IgnoreCase, ESearchDir::FromStart);
-				OriginalRootString = TEXT("/") + OriginalRootString;
-				RecursiveGetDependencies(*PackageIt, AllPackageNamesToMove, OriginalRootString, ExternalObjectsPaths);
+				RecursiveGetDependencies(*PackageIt, AllPackageNamesToMove, ExternalObjectsPaths);
 			}
 		}
 	}
-
+	
 	// Confirm that there is at least one package to move 
-	if ( AllPackageNamesToMove.Num() == 0 )
+	if (AllPackageNamesToMove.Num() == 0)
 	{
 		FMessageDialog::Open(EAppMsgType::Ok, LOCTEXT("MigratePackages_NoFilesToMove", "No files were found to move"));
 		return;
@@ -2895,6 +2894,288 @@ void UAssetToolsImpl::MigratePackages_ReportConfirmed(TSharedPtr<TArray<ReportPa
 		return;
 	}
 
+	// Build a list of packages to handle
+	TSet<FName> AllPackageNamesToMove;
+	for (auto PackageDataIt = PackageDataToMigrate->CreateConstIterator(); PackageDataIt; ++PackageDataIt)
+	{
+		if (PackageDataIt->bShouldMigratePackage)
+		{
+			AllPackageNamesToMove.Add(FName(PackageDataIt->Name));
+		}
+	}
+
+	bool bAbort = false;
+	FMessageLog MigrateLog("AssetTools");
+
+	// Determine if the destination is a project content folder or a plugin content folder
+	TArray<FString> ProjectFiles;
+	IFileManager::Get().FindFiles(ProjectFiles, *(DestinationFolder + TEXT("../")), TEXT("uproject"));
+	bool bIsDestinationAProject = !ProjectFiles.IsEmpty();
+
+	// Associate each Content folder in the target Plugin hierarchy to a content root string in UFS
+	TMap<FName, FString> DestContentRootsToFolders;
+
+	// Assets in /Game always map directly to the destination
+	DestContentRootsToFolders.Add(FName(TEXT("/Game")), DestinationFolder);
+
+	// If our destination is a project, it could have plugins...
+	if (bIsDestinationAProject)
+	{
+		// Find all "Content" folders under the destination ../Plugins directory
+		TArray<FString> ContentFolders;
+		IFileManager::Get().FindFilesRecursive(ContentFolders, *(DestinationFolder + TEXT("../Plugins/")), TEXT("Content"), false, true);
+
+		for (const FString& Folder : ContentFolders)
+		{
+			// Parse the parent folder of .../Content from "Folder"
+			FString Path, Content, Root;
+			bool bSplitContent = Folder.Split(TEXT("/"), &Path, &Content, ESearchCase::CaseSensitive, ESearchDir::FromEnd);
+			bool bSplitParentPath = Path.Split(TEXT("/"), nullptr, &Root, ESearchCase::CaseSensitive, ESearchDir::FromEnd);
+			if (!bSplitContent || !bSplitParentPath)
+			{
+				MigrateLog.Error(FText::Format(LOCTEXT("MigratePackages_NoMountPointFolder", "Unable to determine mount point for folder {0}"), FText::FromString(Folder)));
+				bAbort = true;
+				continue;
+			}
+
+			// Determine this folder name to be a content root in the destination
+			FString DestContentRoot = FString(TEXT("/")) + Root;
+			FString DestContentFolder = FPaths::ConvertRelativePathToFull(Folder + TEXT("/"));
+			DestContentRootsToFolders.Add(FName(DestContentRoot), DestContentFolder);
+		}
+	}
+
+	if (bAbort)
+	{
+		MigrateLog.Notify();
+		return;
+	}
+
+	// Check if the root of any of the packages to migrate have no destination
+	TArray<FName> LostPackages;
+	TSet<FName> LostPackageRoots;
+	for (const FName& PackageName : AllPackageNamesToMove)
+	{
+		// Acquire the mount point for this package
+		FString Folder = TEXT("/") + FPackageName::GetPackageMountPoint(PackageName.ToString()).ToString();
+
+		// If this is /Game package, it doesn't need special handling and is simply bound for the destination, continue
+		if (Folder == TEXT("/Game"))
+		{
+			continue;
+		}
+
+		// Resolve the disk folder of this package's mount point so we compare directory names directly
+		//  We do this FileSystem to FileSystem compare instead of mount point (UFS) to content root (FileSystem)
+		//  The mount point name likely comes from the FriendlyName in the uplugin, or the uplugin basename
+		//  What's important here is that we succeed in finding _the same plugin_ between a copy/pasted project
+		if (!FPackageName::TryConvertLongPackageNameToFilename(Folder, Folder))
+		{
+			MigrateLog.Error(FText::Format(LOCTEXT("MigratePackages_NoContentFolder", "Unable to determine content folder for asset {0}"), FText::FromString(PackageName.ToString())));
+			bAbort = true;
+			continue;
+		}
+		Folder.RemoveFromEnd(TEXT("/"));
+
+		// Parse the parent folder of .../Content from "Folder"
+		FString Path, Content, Root;
+		bool bSplitContent = Folder.Split(TEXT("/"), &Path, &Content, ESearchCase::CaseSensitive, ESearchDir::FromEnd);
+		bool bSplitParentPath = Path.Split(TEXT("/"), nullptr, &Root, ESearchCase::CaseSensitive, ESearchDir::FromEnd);
+		if (!bSplitContent || !bSplitParentPath)
+		{
+			MigrateLog.Error(FText::Format(LOCTEXT("MigratePackages_NoMountPointPackage", "Unable to determine mount point for package {0}"), FText::FromString(PackageName.ToString())));
+			bAbort = true;
+			continue;
+		}
+
+		// Check to see if the content root exists in the destination, otherwise it's "Lost"
+		FString SrcContentRoot = FString(TEXT("/")) + Root;
+		FName SrcContentRootName(SrcContentRoot);
+		if (!DestContentRootsToFolders.Find(SrcContentRootName))
+		{
+			LostPackages.Add(PackageName);
+			LostPackageRoots.Add(SrcContentRootName);
+		}
+	}
+
+	if (bAbort)
+	{
+		MigrateLog.Notify();
+		return;
+	}
+
+	// If some packages don't have a matching content root in the destination, prompt for desired behavior
+	if (!LostPackages.IsEmpty())
+	{
+		FString LostPackageRootsString;
+		for (const FName& PackageRoot : LostPackageRoots)
+		{
+			LostPackageRootsString += FString(TEXT("\n\t")) + PackageRoot.ToString();
+		}
+
+		// Prompt to consolidate to a migration folder
+		FText Prompt = FText::Format(LOCTEXT("MigratePackages_ConsolidateToTemp", "Some selected assets don't have a corresponding content root in the destination.{0}\n\nWould you like to migrate all selected assets into a folder with consolidated references? Without migrating into a folder the assets in the above roots will not be migrated."), FText::FromString(LostPackageRootsString));
+		switch (FMessageDialog::Open(EAppMsgType::YesNoCancel, Prompt))
+		{
+		case EAppReturnType::Yes:
+			// No op
+			break;
+		case EAppReturnType::No:
+			LostPackages.Reset();
+			break;
+		case EAppReturnType::Cancel:
+			return;
+		}
+	}
+
+	// This will be used to tidy up the temp folder after we copy the migrated assets
+	FString SrcDiskFolderFilename;
+
+	// Fixing up references requires resaving packages to a temporary location
+	if (!LostPackages.IsEmpty())
+	{
+		// Resolve the packages to migrate to assets
+		TArray<UObject*> SrcObjects;
+		for (const FName& SrcPackage : AllPackageNamesToMove)
+		{
+			UPackage* LoadedPackage = UPackageTools::LoadPackage(SrcPackage.ToString());
+			if (!LoadedPackage)
+			{
+				MigrateLog.Error(FText::Format(LOCTEXT("MigratePackages_FailedToLoadPackage", "Failed to load package {0}"), FText::FromString(SrcPackage.ToString())));
+				bAbort = true;
+				continue;
+			}
+
+			UObject* Asset = LoadedPackage->FindAssetInPackage();
+			if (Asset)
+			{
+				SrcObjects.Add(Asset);
+			}
+			else
+			{
+				MigrateLog.Warning(FText::Format(LOCTEXT("MigratePackages_PackageHasNoAsset", "Package {0} has no asset in it"), FText::FromString(SrcPackage.ToString())));
+			}
+		}
+
+		if (bAbort)
+		{
+			MigrateLog.Notify();
+			return;
+		}
+
+		// Query the user for a folder to migrate assets into. This folder will exist temporary in this project, and is the destination for migration in the target project
+		FString FolderName;
+		bool bIsOkButtonEnabled = true;
+		TSharedRef<SEditableTextBox> EditableTextBox = SNew(SEditableTextBox)
+			.Text(FText::FromString(TEXT("Migrated")))
+			.OnVerifyTextChanged_Lambda([&bIsOkButtonEnabled](const FText& InNewText, FText& OutErrorMessage) -> bool
+				{
+					if (InNewText.ToString().Contains(TEXT("/")))
+					{
+						OutErrorMessage = LOCTEXT("Migrated_CannotContainSlashes", "Cannot use a slash in a folder name.");
+
+						// Disable Ok if the string is invalid
+						bIsOkButtonEnabled = false;
+						return false;
+					}
+
+					// Enable Ok if the string is valid
+					bIsOkButtonEnabled = true;
+					return true;
+				})
+			.OnTextCommitted_Lambda([&FolderName](const FText& NewValue, ETextCommit::Type)
+				{
+					// Set the result if they modified the text
+					FolderName = NewValue.ToString();
+				});
+
+		// Set the result if they just click Ok
+		SGenericDialogWidget::FArguments FolderDialogArguments;
+		FolderDialogArguments.OnOkPressed_Lambda([&EditableTextBox, &FolderName]()
+			{
+				FolderName = EditableTextBox->GetText().ToString();
+			});
+
+		// Present the Dialog
+		SGenericDialogWidget::OpenDialog(LOCTEXT("MigratePackages_FolderName", "Folder for Migrated Assets"),
+			SNew(SVerticalBox)
+			+ SVerticalBox::Slot()
+				.AutoHeight()
+				.Padding(5.0f)
+			[
+				SNew(STextBlock)
+				.Text(LOCTEXT("MigratePackages_SpecifyConsolidateFolder", "Please specify a new folder name to consolidate the assets into."))
+			]
+			+ SVerticalBox::Slot()
+				.Padding(5.0f)
+			[
+				SNew(SSpacer)
+			]
+			+ SVerticalBox::Slot()
+				.AutoHeight()
+				.Padding(5.0f)
+			[
+				EditableTextBox
+			],
+			FolderDialogArguments, true);
+
+		// Sanity the user input
+		if (FolderName.IsEmpty())
+		{
+			return;
+		}
+
+		// Remove forbidden characters
+		FolderName = FolderName.Replace(TEXT("/"), TEXT(""));
+
+		// Verify that we don't have any assets that exist where we want to perform our consolidation
+		//  We could do the same check at the destination, but the package copy code will happily overwrite,
+		//  and it seems reasonable someone would want to Migrate several times if snapshotting changes from another project.
+		//  So for now we don't necessitate that the destination content folder be missing as well.
+		FString SrcUfsFolderName = FolderName;
+		SrcUfsFolderName.InsertAt(0, TEXT("/Game/"));
+		SrcDiskFolderFilename = FPackageName::LongPackageNameToFilename(SrcUfsFolderName, TEXT(""));
+		if (IFileManager::Get().DirectoryExists(*SrcDiskFolderFilename))
+		{
+			const FText Message = FText::Format(LOCTEXT("MigratePackages_InvalidMigrateFolder", "{0} exists on disk in the source project, and cannot be used to consolidate assets."), FText::FromString(SrcDiskFolderFilename));
+			EAppReturnType::Type Response = FMessageDialog::Open(EAppMsgType::Ok, Message);
+			return;
+		}
+
+		// To handle complex references and assets in different Plugins, we must first duplicate to temp packages, then migrate those temps
+		TArray<UObject*> TempObjects;
+		ObjectTools::DuplicateObjects(SrcObjects, TEXT(""), SrcUfsFolderName, /*bOpenDialog=*/false, &TempObjects);
+		TMap<UObject*, UObject*> ReplacementMap;
+		for (int i=0; i<SrcObjects.Num(); ++i)
+		{
+			ReplacementMap.Add(SrcObjects[i], TempObjects[i]);
+		}
+
+		// Save fixed up packages to the migrated folder, and update the set of files to copy to be those migrated packages
+		{
+			TSet<FName> NewPackageNamesToMove;
+			for (int i=0; i<TempObjects.Num(); ++i)
+			{
+				UObject* TempObject = TempObjects[i];
+
+				// Fixup references in each package, save them, and update the source of our copy operation
+				FArchiveReplaceObjectRef<UObject> ReplaceAr(TempObject, ReplacementMap, EArchiveReplaceObjectFlags::IgnoreOuterRef | EArchiveReplaceObjectFlags::IgnoreArchetypeRef);
+
+				// Calculate the file path to the new, migrated package
+				FString const TempPackageName = TempObject->GetPackage()->GetName();
+				FString const TempPackageFilename = FPackageName::LongPackageNameToFilename(TempPackageName, FPackageName::GetAssetPackageExtension());
+
+				// Save it
+				FSavePackageArgs SaveArgs;
+				GEditor->Save(TempObject->GetPackage(), /*InAsset=*/nullptr, *TempPackageFilename, SaveArgs);
+
+				NewPackageNamesToMove.Add(FName(TempPackageName));
+			}
+
+			AllPackageNamesToMove = NewPackageNamesToMove;
+		}
+	}
+
 	bool bUserCanceled = false;
 
 	// Copy all specified assets and their dependencies to the destination folder
@@ -2909,15 +3190,11 @@ void UAssetToolsImpl::MigratePackages_ReportConfirmed(TSharedPtr<TArray<ReportPa
 	SlowTask.EnterProgressFrame();
 	{
 		FScopedSlowTask LoopProgress(PackageDataToMigrate.Get()->Num());
-		for ( auto PackageDataIt = PackageDataToMigrate.Get()->CreateConstIterator(); PackageDataIt; ++PackageDataIt)
+		for ( const FName& PackageNameToMove : AllPackageNamesToMove )
 		{
 			LoopProgress.EnterProgressFrame();
-			if (!PackageDataIt->bShouldMigratePackage)
-			{
-				continue;
-			}
 
-			const FString& PackageName = PackageDataIt->Name;
+			const FString& PackageName = PackageNameToMove.ToString();
 			FString SrcFilename;
 			
 			if (!FPackageName::DoesPackageExist(PackageName, &SrcFilename))
@@ -2935,7 +3212,22 @@ void UAssetToolsImpl::MigratePackages_ReportConfirmed(TSharedPtr<TArray<ReportPa
 			{
 				bool bFileOKToCopy = true;
 
-				FString DestFilename = DestinationFolder;
+				FString Path = PackageNameToMove.ToString();
+				FString PackageRoot;
+				Path.RemoveFromStart(TEXT("/"));
+				Path.Split("/", &PackageRoot, &Path, ESearchCase::IgnoreCase, ESearchDir::FromStart);
+				PackageRoot = TEXT("/") + PackageRoot;
+
+				FString DestFilename;
+				TMap<FName, FString>::ValueType* DestRootFolder = DestContentRootsToFolders.Find(FName(PackageRoot));
+				if (ensure(DestRootFolder))
+				{
+					DestFilename = *DestRootFolder;
+				}
+				else
+				{
+					continue;
+				}
 
 				FString SubFolder;
 				if ( SrcFilename.Split( TEXT("/Content/"), nullptr, &SubFolder ) )
@@ -2994,6 +3286,24 @@ void UAssetToolsImpl::MigratePackages_ReportConfirmed(TSharedPtr<TArray<ReportPa
 		}
 	}
 
+	// If we are consolidating lost packages, we are copying temporary packages, so clean them up.
+	if (!LostPackages.IsEmpty())
+	{
+		TArray<FAssetData> AssetsToDelete;
+		for (const FName& PackageNameToMove : AllPackageNamesToMove)
+		{
+			AssetsToDelete.Add(FAssetData(UPackageTools::LoadPackage(PackageNameToMove.ToString())));
+		}
+
+		ObjectTools::DeleteAssets(AssetsToDelete, /*bShowConfirmation=*/false);
+
+		if (!IFileManager::Get().DeleteDirectory(*SrcDiskFolderFilename))
+		{
+			UE_LOG(LogAssetTools, Warning, TEXT("Failed to delete temporary directory %s while migrating assets"), *SrcDiskFolderFilename);
+			CopyErrors += SrcDiskFolderFilename + LINE_TERMINATOR;
+		}
+	}
+
 	FString SourceControlErrors;
 	SlowTask.EnterProgressFrame();
 
@@ -3023,7 +3333,6 @@ void UAssetToolsImpl::MigratePackages_ReportConfirmed(TSharedPtr<TArray<ReportPa
 		}
 	}
 
-	FMessageLog MigrateLog("AssetTools");
 	FText LogMessage = FText::FromString(TEXT("Content migration completed successfully!"));
 	EMessageSeverity::Type Severity = EMessageSeverity::Info;
 	if ( CopyErrors.Len() > 0 || SourceControlErrors.Len() > 0 )
@@ -3077,7 +3386,7 @@ void UAssetToolsImpl::MigratePackages_ReportConfirmed(TSharedPtr<TArray<ReportPa
 	MigrateLog.Notify(LogMessage, Severity, true);
 }
 
-void UAssetToolsImpl::RecursiveGetDependencies(const FName& PackageName, TSet<FName>& AllDependencies, const FString& OriginalRoot, TSet<FString>& OutExternalObjectsPaths) const
+void UAssetToolsImpl::RecursiveGetDependencies(const FName& PackageName, TSet<FName>& AllDependencies, TSet<FString>& OutExternalObjectsPaths) const
 {
 	FAssetRegistryModule& AssetRegistryModule = FModuleManager::Get().LoadModuleChecked<FAssetRegistryModule>(TEXT("AssetRegistry"));
 	TArray<FName> Dependencies;
@@ -3085,16 +3394,16 @@ void UAssetToolsImpl::RecursiveGetDependencies(const FName& PackageName, TSet<FN
 	
 	for ( auto DependsIt = Dependencies.CreateConstIterator(); DependsIt; ++DependsIt )
 	{
-		if ( !AllDependencies.Contains(*DependsIt) )
+		FString DependencyName = (*DependsIt).ToString();
+
+		const bool bIsEnginePackage = DependencyName.StartsWith(TEXT("/Engine"));
+		const bool bIsScriptPackage = DependencyName.StartsWith(TEXT("/Script"));
+		if ( !bIsEnginePackage && !bIsScriptPackage )
 		{
-			const bool bIsEnginePackage = (*DependsIt).ToString().StartsWith(TEXT("/Engine"));
-			const bool bIsScriptPackage = (*DependsIt).ToString().StartsWith(TEXT("/Script"));
-			// Skip all packages whose root is different than the source package list root
-			const bool bIsInSamePackage = (*DependsIt).ToString().StartsWith(OriginalRoot);
-			if ( !bIsEnginePackage && !bIsScriptPackage && bIsInSamePackage )
+			if (!AllDependencies.Contains(*DependsIt))
 			{
 				AllDependencies.Add(*DependsIt);
-				RecursiveGetDependencies(*DependsIt, AllDependencies, OriginalRoot, OutExternalObjectsPaths);
+				RecursiveGetDependencies(*DependsIt, AllDependencies, OutExternalObjectsPaths);
 			}
 		}
 	}
@@ -3122,14 +3431,13 @@ void UAssetToolsImpl::RecursiveGetDependencies(const FName& PackageName, TSet<FN
 						for (const FAssetData& ExternalObjectAsset : ExternalObjectAssets)
 						{
 							AllDependencies.Add(ExternalObjectAsset.PackageName);
-							RecursiveGetDependencies(ExternalObjectAsset.PackageName, AllDependencies, OriginalRoot, OutExternalObjectsPaths);
+							RecursiveGetDependencies(ExternalObjectAsset.PackageName, AllDependencies, OutExternalObjectsPaths);
 						}
 					}
 				}
 			}
 		}
 	}
-	
 }
 
 void UAssetToolsImpl::RecursiveGetDependenciesAdvanced(const FName& PackageName, FAdvancedCopyParams& CopyParams, TArray<FName>& AllDependencies, TMap<FName, FName>& DependencyMap, const UAdvancedCopyCustomization* CopyCustomization, TArray<FAssetData>& OptionalAssetData) const
diff --git a/Engine/Source/Developer/AssetTools/Private/AssetTools.h b/Engine/Source/Developer/AssetTools/Private/AssetTools.h
--- a/Engine/Source/Developer/AssetTools/Private/AssetTools.h
+++ b/Engine/Source/Developer/AssetTools/Private/AssetTools.h
@@ -159,7 +159,7 @@ private:
 	void AdvancedCopyPackages_ReportConfirmed(FAdvancedCopyParams CopyParam, TArray<TMap<FString, FString>> DestinationMap) const;
 
 	/** Gets the dependencies of the specified package recursively */
-	void RecursiveGetDependencies(const FName& PackageName, TSet<FName>& AllDependencies, const FString& OriginalRoot, TSet<FString>& ExternalObjectsPaths) const;
+	void RecursiveGetDependencies(const FName& PackageName, TSet<FName>& AllDependencies, TSet<FString>& ExternalObjectsPaths) const;
 
 	/** Gets the dependencies of the specified package recursively while omitting things that don't pass the FARFilter passed in from FAdvancedCopyParams */
 	void RecursiveGetDependenciesAdvanced(const FName& PackageName, FAdvancedCopyParams& CopyParams, TArray<FName>& AllDependencies, TMap<FName, FName>& DependencyMap, const class UAdvancedCopyCustomization* CopyCustomization, TArray<FAssetData>& OptionalAssetData) const;
