Fixed crash opening a child anim blueprint of a template anim blueprint

Ensure that template flag is propagated in the regular factory use case

diff --git a/Engine/Source/Developer/AssetTools/Private/AssetTypeActions/AssetTypeActions_AnimBlueprint.cpp b/Engine/Source/Developer/AssetTools/Private/AssetTypeActions/AssetTypeActions_AnimBlueprint.cpp
--- a/Engine/Source/Developer/AssetTools/Private/AssetTypeActions/AssetTypeActions_AnimBlueprint.cpp
+++ b/Engine/Source/Developer/AssetTools/Private/AssetTypeActions/AssetTypeActions_AnimBlueprint.cpp
@@ -70,6 +70,7 @@ void FAssetTypeActions_AnimBlueprint::GetActions(const TArray<UObject*>& InObjec
 							UAnimBlueprintFactory* AnimBlueprintFactory = NewObject<UAnimBlueprintFactory>();
 							AnimBlueprintFactory->ParentClass = TSubclassOf<UAnimInstance>(*TargetParentBP->GeneratedClass);
 							AnimBlueprintFactory->TargetSkeleton = TargetSkeleton;
+							AnimBlueprintFactory->bTemplate = false;
 
 							FContentBrowserModule& ContentBrowserModule = FModuleManager::LoadModuleChecked<FContentBrowserModule>("ContentBrowser");
 							ContentBrowserModule.Get().CreateNewAsset(Name, PackagePath, TargetParentBP->GetClass(), AnimBlueprintFactory);
@@ -169,6 +170,7 @@ UFactory* FAssetTypeActions_AnimBlueprint::GetFactoryForBlueprintType(UBlueprint
 		UAnimBlueprintFactory* AnimBlueprintFactory = NewObject<UAnimBlueprintFactory>();
 		AnimBlueprintFactory->ParentClass = TSubclassOf<UAnimInstance>(*InBlueprint->GeneratedClass);
 		AnimBlueprintFactory->TargetSkeleton = AnimBlueprint->TargetSkeleton;
+		AnimBlueprintFactory->bTemplate = AnimBlueprint->bIsTemplate;
 		return AnimBlueprintFactory;
 	}
 }
