Migrate: Fix flattening of assets into the target folder during consolidate

DuplicateObjects takes the destination path very literally, flattening objects to a single folder. Instead, loop over all the objects, calculate the correct relative path, and duplicate the objects one-by-one.

#rnx

diff --git a/Engine/Source/Developer/AssetTools/Private/AssetTools.cpp b/Engine/Source/Developer/AssetTools/Private/AssetTools.cpp
--- a/Engine/Source/Developer/AssetTools/Private/AssetTools.cpp
+++ b/Engine/Source/Developer/AssetTools/Private/AssetTools.cpp
@@ -3145,11 +3145,37 @@ void UAssetToolsImpl::MigratePackages_ReportConfirmed(TSharedPtr<TArray<ReportPa
 
 		// To handle complex references and assets in different Plugins, we must first duplicate to temp packages, then migrate those temps
 		TArray<UObject*> TempObjects;
-		ObjectTools::DuplicateObjects(SrcObjects, TEXT(""), SrcUfsFolderName, /*bOpenDialog=*/false, &TempObjects);
 		TMap<UObject*, UObject*> ReplacementMap;
+		TSet<UPackage*> PackagesUserRefusedToFullyLoad;
 		for (int i=0; i<SrcObjects.Num(); ++i)
 		{
-			ReplacementMap.Add(SrcObjects[i], TempObjects[i]);
+			UObject* Object = SrcObjects[i];
+			FString PackageName = Object->GetPackage()->GetName();
+
+			FString Path, Root;
+			PackageName.RemoveFromStart(TEXT("/"));
+			bool bSplitRoot = PackageName.Split(TEXT("/"), &Root, &Path);
+			if (!bSplitRoot)
+			{
+				MigrateLog.Error(FText::Format(LOCTEXT("MigratePackages_NoMountPointPackage", "Unable to determine mount point for package {0}"), FText::FromString(PackageName)));
+				continue;
+			}
+
+			ObjectTools::FPackageGroupName PackageGroupName;
+			PackageGroupName.ObjectName = Object->GetName();
+			PackageGroupName.PackageName = SrcUfsFolderName + TEXT("/") + Path;
+			FString GroupName = Object->GetFullGroupName(/*bStartWithOuter =*/true);
+			if (GroupName != TEXT("None"))
+			{
+				PackageGroupName.GroupName = GroupName;
+			}
+
+			UObject* NewObject = ObjectTools::DuplicateSingleObject(Object, PackageGroupName, PackagesUserRefusedToFullyLoad);
+			if (NewObject)
+			{
+				TempObjects.Add(NewObject);
+				ReplacementMap.Add(SrcObjects[i], TempObjects[i]);
+			}
 		}
 
 		// Save fixed up packages to the migrated folder, and update the set of files to copy to be those migrated packages
