Add SlowTask dialogs, Fix cleaning up temp assets

#rnx

diff --git a/Engine/Source/Developer/AssetTools/Private/AssetTools.cpp b/Engine/Source/Developer/AssetTools/Private/AssetTools.cpp
--- a/Engine/Source/Developer/AssetTools/Private/AssetTools.cpp
+++ b/Engine/Source/Developer/AssetTools/Private/AssetTools.cpp
@@ -3143,51 +3143,75 @@ void UAssetToolsImpl::MigratePackages_ReportConfirmed(TSharedPtr<TArray<ReportPa
 			return;
 		}
 
-		// To handle complex references and assets in different Plugins, we must first duplicate to temp packages, then migrate those temps
 		TArray<UObject*> TempObjects;
 		TMap<UObject*, UObject*> ReplacementMap;
-		TSet<UPackage*> PackagesUserRefusedToFullyLoad;
-		for (int i=0; i<SrcObjects.Num(); ++i)
+
+		// Copy all specified assets and their dependencies to the destination folder
+		FScopedSlowTask SlowTask( 3, LOCTEXT( "MigratePackages_Consolidate", "Consolidating Assets..." ) );
+		SlowTask.MakeDialog();
+
+		// To handle complex references and assets in different Plugins, we must first duplicate to temp packages
 		{
-			UObject* Object = SrcObjects[i];
-			FString PackageName = Object->GetPackage()->GetName();
+			FScopedSlowTask LoopProgress(SrcObjects.Num());
 
-			FString Path, Root;
-			PackageName.RemoveFromStart(TEXT("/"));
-			bool bSplitRoot = PackageName.Split(TEXT("/"), &Root, &Path);
-			if (!bSplitRoot)
+			TSet<UPackage*> PackagesUserRefusedToFullyLoad;
+			for (int i=0; i<SrcObjects.Num(); ++i)
 			{
-				MigrateLog.Error(FText::Format(LOCTEXT("MigratePackages_NoMountPointPackage", "Unable to determine mount point for package {0}"), FText::FromString(PackageName)));
-				continue;
-			}
+				LoopProgress.EnterProgressFrame();
 
-			ObjectTools::FPackageGroupName PackageGroupName;
-			PackageGroupName.ObjectName = Object->GetName();
-			PackageGroupName.PackageName = SrcUfsFolderName + TEXT("/") + Path;
-			FString GroupName = Object->GetFullGroupName(/*bStartWithOuter =*/true);
-			if (GroupName != TEXT("None"))
-			{
-				PackageGroupName.GroupName = GroupName;
+				UObject* Object = SrcObjects[i];
+				FString PackageName = Object->GetPackage()->GetName();
+
+				FString Path, Root;
+				PackageName.RemoveFromStart(TEXT("/"));
+				bool bSplitRoot = PackageName.Split(TEXT("/"), &Root, &Path);
+				if (!bSplitRoot)
+				{
+					MigrateLog.Error(FText::Format(LOCTEXT("MigratePackages_NoMountPointPackage", "Unable to determine mount point for package {0}"), FText::FromString(PackageName)));
+					continue;
+				}
+
+				ObjectTools::FPackageGroupName PackageGroupName;
+				PackageGroupName.ObjectName = Object->GetName();
+				PackageGroupName.PackageName = SrcUfsFolderName + TEXT("/") + Path;
+				FString GroupName = Object->GetFullGroupName(/*bStartWithOuter =*/true);
+				if (GroupName != TEXT("None"))
+				{
+					PackageGroupName.GroupName = GroupName;
+				}
+
+				UObject* NewObject = ObjectTools::DuplicateSingleObject(Object, PackageGroupName, PackagesUserRefusedToFullyLoad);
+				if (NewObject)
+				{
+					TempObjects.Add(NewObject);
+					ReplacementMap.Add(SrcObjects[i], TempObjects[i]);
+				}
 			}
+		}
+
+		// Update references between TempObjects (to reference each other)
+		{
+			FScopedSlowTask LoopProgress(TempObjects.Num());
 
-			UObject* NewObject = ObjectTools::DuplicateSingleObject(Object, PackageGroupName, PackagesUserRefusedToFullyLoad);
-			if (NewObject)
+			for (int i=0; i<TempObjects.Num(); ++i)
 			{
-				TempObjects.Add(NewObject);
-				ReplacementMap.Add(SrcObjects[i], TempObjects[i]);
+				LoopProgress.EnterProgressFrame();
+				UObject* TempObject = TempObjects[i];
+
+				FArchiveReplaceObjectRef<UObject> ReplaceAr(TempObject, ReplacementMap, EArchiveReplaceObjectFlags::IgnoreOuterRef | EArchiveReplaceObjectFlags::IgnoreArchetypeRef);
 			}
 		}
 
 		// Save fixed up packages to the migrated folder, and update the set of files to copy to be those migrated packages
 		{
+			FScopedSlowTask LoopProgress(TempObjects.Num());
+
 			TSet<FName> NewPackageNamesToMove;
 			for (int i=0; i<TempObjects.Num(); ++i)
 			{
+				LoopProgress.EnterProgressFrame();
 				UObject* TempObject = TempObjects[i];
 
-				// Fixup references in each package, save them, and update the source of our copy operation
-				FArchiveReplaceObjectRef<UObject> ReplaceAr(TempObject, ReplacementMap, EArchiveReplaceObjectFlags::IgnoreOuterRef | EArchiveReplaceObjectFlags::IgnoreArchetypeRef);
-
 				// Calculate the file path to the new, migrated package
 				FString const TempPackageName = TempObject->GetPackage()->GetName();
 				FString const TempPackageFilename = FPackageName::LongPackageNameToFilename(TempPackageName, FPackageName::GetAssetPackageExtension());
@@ -3319,12 +3343,17 @@ void UAssetToolsImpl::MigratePackages_ReportConfirmed(TSharedPtr<TArray<ReportPa
 		TArray<FAssetData> AssetsToDelete;
 		for (const FName& PackageNameToMove : AllPackageNamesToMove)
 		{
-			AssetsToDelete.Add(FAssetData(UPackageTools::LoadPackage(PackageNameToMove.ToString())));
+			UPackage* Package = UPackageTools::LoadPackage(PackageNameToMove.ToString());
+			if (Package)
+			{
+				FAssetData AssetData(Package);
+				AssetsToDelete.Add(AssetData);
+			}
 		}
 
 		ObjectTools::DeleteAssets(AssetsToDelete, /*bShowConfirmation=*/false);
 
-		if (!IFileManager::Get().DeleteDirectory(*SrcDiskFolderFilename))
+		if (!IFileManager::Get().DeleteDirectory(*SrcDiskFolderFilename, /*RequireExists =*/ false, /*Tree =*/ true))
 		{
 			UE_LOG(LogAssetTools, Warning, TEXT("Failed to delete temporary directory %s while migrating assets"), *SrcDiskFolderFilename);
 			CopyErrors += SrcDiskFolderFilename + LINE_TERMINATOR;
