Treat Engine content more intelligently when Migrating/Consolidating assets.

### AssetTools
Check each package's mount point and fold that against each Plugin's LoadedFrom (Project or Engine)
If Engine content is visible in the Content Browser, show but do not schedule Engine Plugin content for Migration
Add support for sending in "ShouldMigratePackage" state to SPackageReportDialog

### SPackageReportDialog
Remove "Active" state -> this dialog used to associate Disabled tree item state with the state of not Migrating a package. Just rely on the check state to dictate if a package will be Migrated
Add support for indeterminate checkbox state for folders within the dialog when some, but not all, children being selected for migration
Add logic for initial check/indeterminate/unchecked state and when check state changes on an item
Propagate check state up the hierarchy from the check change event item

#rnx

diff --git a/Engine/Source/Developer/AssetTools/Private/AssetTools.cpp b/Engine/Source/Developer/AssetTools/Private/AssetTools.cpp
--- a/Engine/Source/Developer/AssetTools/Private/AssetTools.cpp
+++ b/Engine/Source/Developer/AssetTools/Private/AssetTools.cpp
@@ -155,6 +155,8 @@
 #include "Dialogs/Dialogs.h"
 #include "Widgets/Input/SEditableTextBox.h"
 #include "Widgets/Layout/SSpacer.h"
+#include "Interfaces/IPluginManager.h"
+#include "Settings/ContentBrowserSettings.h"
 
 #if WITH_EDITOR
 #include "Subsystems/AssetEditorSubsystem.h"
@@ -2820,7 +2822,67 @@ void UAssetToolsImpl::PerformMigratePackages(TArray<FName> PackageNamesToMigrate
 			}
 		}
 	}
-	
+
+	// Fetch the enabled plugins and their mount points
+	TMap<FName, EPluginLoadedFrom> EnabledPluginToLoadedFrom;
+	TArray<TSharedRef<IPlugin>> EnabledPlugins = IPluginManager::Get().GetEnabledPluginsWithContent();
+	for (const TSharedRef<IPlugin>& EnabledPlugin : EnabledPlugins)
+	{
+		EnabledPluginToLoadedFrom.Add(FName(EnabledPlugin->GetMountedAssetPath()), EnabledPlugin->GetLoadedFrom());
+	}
+
+	// Find assets in non-Project Plugins
+	TSet<FName> ShouldMigratePackage;
+	bool bShouldShowEngineContent = GetDefault<UContentBrowserSettings>()->GetDisplayEngineFolder();
+	{
+		// This is the new list to prompt for migration
+		TSet<FName> FilteredPackageNamesToMove;
+
+		for (const FName& PackageName : AllPackageNamesToMove)
+		{
+			 FName PackageMountPoint = FPackageName::GetPackageMountPoint(PackageName.ToString(), false);
+			 EPluginLoadedFrom* Found = EnabledPluginToLoadedFrom.Find(PackageMountPoint);
+
+			 bool bShouldMigratePackage = true;
+			 if (Found)
+			 {
+				 // plugin content, decide if it's appropriate to migrate
+				 switch (*Found)
+				 {
+				 case EPluginLoadedFrom::Engine:
+					 if (!bShouldShowEngineContent)
+					 {
+						 continue;
+					 }
+					 bShouldMigratePackage = false;
+					 break;
+
+				 case EPluginLoadedFrom::Project:
+					 bShouldMigratePackage = true;
+					 break;
+				 
+				 default:
+					 bShouldMigratePackage = false;
+					 break;
+				 }
+			 }
+			 else
+			 {
+				 // this is not plugin content
+				 bShouldMigratePackage = true;
+			 }
+
+			 FilteredPackageNamesToMove.Add(PackageName);
+
+			 if (bShouldMigratePackage)
+			 {
+				 ShouldMigratePackage.Add(PackageName);
+			 }
+		}
+
+		AllPackageNamesToMove = FilteredPackageNamesToMove;
+	}
+
 	// Confirm that there is at least one package to move 
 	if (AllPackageNamesToMove.Num() == 0)
 	{
@@ -2834,7 +2896,8 @@ void UAssetToolsImpl::PerformMigratePackages(TArray<FName> PackageNamesToMigrate
 		TSharedPtr<TArray<ReportPackageData>> ReportPackages = MakeShareable(new TArray<ReportPackageData>);
 		for ( auto PackageIt = AllPackageNamesToMove.CreateConstIterator(); PackageIt; ++PackageIt )
 		{
-			ReportPackages.Get()->Add({ (*PackageIt).ToString(), true });
+			bool bShouldMigratePackage = ShouldMigratePackage.Find(*PackageIt) != nullptr;
+			ReportPackages.Get()->Add({ (*PackageIt).ToString(), bShouldMigratePackage });
 		}
 		SPackageReportDialog::FOnReportConfirmed OnReportConfirmed = SPackageReportDialog::FOnReportConfirmed::CreateUObject(this, &UAssetToolsImpl::MigratePackages_ReportConfirmed, ReportPackages);
 		SPackageReportDialog::OpenPackageReportDialog(ReportMessage, *ReportPackages.Get(), OnReportConfirmed);
@@ -3015,7 +3078,7 @@ void UAssetToolsImpl::MigratePackages_ReportConfirmed(TSharedPtr<TArray<ReportPa
 		}
 
 		// Prompt to consolidate to a migration folder
-		FText Prompt = FText::Format(LOCTEXT("MigratePackages_ConsolidateToTemp", "Some selected assets don't have a corresponding content root in the destination.{0}\n\nWould you like to migrate all selected assets into a folder with consolidated references? Without migrating into a folder the assets in the above roots will not be migrated."), FText::FromString(LostPackageRootsString));
+		FText Prompt = FText::Format(LOCTEXT("MigratePackages_ConsolidateToTemp", "Some selected assets don't have a corresponding content root in the destination.{0}\n\nWould you like to save a copy of all selected assets into a folder with consolidated references? If you select No then assets in the above roots will not be migrated."), FText::FromString(LostPackageRootsString));
 		switch (FMessageDialog::Open(EAppMsgType::YesNoCancel, Prompt))
 		{
 		case EAppReturnType::Yes:
@@ -3240,7 +3303,7 @@ void UAssetToolsImpl::MigratePackages_ReportConfirmed(TSharedPtr<TArray<ReportPa
 
 	SlowTask.EnterProgressFrame();
 	{
-		FScopedSlowTask LoopProgress(PackageDataToMigrate.Get()->Num());
+		FScopedSlowTask LoopProgress(AllPackageNamesToMove.Num());
 		for ( const FName& PackageNameToMove : AllPackageNamesToMove )
 		{
 			LoopProgress.EnterProgressFrame();
diff --git a/Engine/Source/Developer/AssetTools/Private/SPackageReportDialog.cpp b/Engine/Source/Developer/AssetTools/Private/SPackageReportDialog.cpp
--- a/Engine/Source/Developer/AssetTools/Private/SPackageReportDialog.cpp
+++ b/Engine/Source/Developer/AssetTools/Private/SPackageReportDialog.cpp
@@ -25,8 +25,7 @@ struct FCompareFPackageReportNodeByName
 };
 
 FPackageReportNode::FPackageReportNode()
-	: bIsChecked(true)
-	, bIsActive(true)
+	: CheckedState(ECheckBoxState::Undetermined)
 	, bShouldMigratePackage(nullptr)
 	, bIsFolder(false)
 	, Parent(nullptr)
@@ -34,8 +33,7 @@ FPackageReportNode::FPackageReportNode()
 
 FPackageReportNode::FPackageReportNode(const FString& InNodeName, bool InIsFolder)
 	: NodeName(InNodeName)
-	, bIsChecked(true)
-	, bIsActive(true)
+	, CheckedState(ECheckBoxState::Undetermined)
 	, bShouldMigratePackage(nullptr)
 	, bIsFolder(InIsFolder)
 	, Parent(nullptr)
@@ -46,20 +44,24 @@ void FPackageReportNode::AddPackage(const FString& PackageName, bool* bInShouldM
 	TArray<FString> PathElements;
 	PackageName.ParseIntoArray(PathElements, TEXT("/"), /*InCullEmpty=*/true);
 
-	return AddPackage_Recursive(PathElements, bInShouldMigratePackage);
+	(void)AddPackage_Recursive(PathElements, bInShouldMigratePackage);
 }
 
 void FPackageReportNode::ExpandChildrenRecursively(const TSharedRef<PackageReportTree>& Treeview)
 {
 	for ( auto ChildIt = Children.CreateConstIterator(); ChildIt; ++ChildIt )
 	{
-		Treeview->SetItemExpansion(*ChildIt, true);
+		Treeview->SetItemExpansion(*ChildIt, (*ChildIt)->CheckedState != ECheckBoxState::Unchecked);
 		(*ChildIt)->ExpandChildrenRecursively(Treeview);
 	}
 }
 
-void FPackageReportNode::AddPackage_Recursive(TArray<FString>& PathElements, bool* bInShouldMigratePackage)
+FPackageReportNode::FChildrenState FPackageReportNode::AddPackage_Recursive(TArray<FString>& PathElements, bool* bInShouldMigratePackage)
 {
+	FChildrenState ChildrenState;
+	ChildrenState.bAnyChildIsChecked = false;
+	ChildrenState.bAllChildrenAreChecked = true;
+
 	if ( PathElements.Num() > 0 )
 	{
 		// Pop the bottom element
@@ -89,13 +91,21 @@ void FPackageReportNode::AddPackage_Recursive(TArray<FString>& PathElements, boo
 
 		if ( ensure(Child.IsValid()) )
 		{
-			Child->AddPackage_Recursive(PathElements, bInShouldMigratePackage);
+			FChildrenState ChildChildrenState = Child->AddPackage_Recursive(PathElements, bInShouldMigratePackage);
+			ChildrenState.bAnyChildIsChecked |= ChildChildrenState.bAnyChildIsChecked;
+			ChildrenState.bAllChildrenAreChecked &= ChildChildrenState.bAllChildrenAreChecked;
 		}
+
+		CheckedState = ChildrenState.bAllChildrenAreChecked ? ECheckBoxState::Checked : (ChildrenState.bAnyChildIsChecked ? ECheckBoxState::Undetermined : ECheckBoxState::Unchecked);
 	}
-	else 
+	else
 	{
+		CheckedState = *bInShouldMigratePackage ? ECheckBoxState::Checked : ECheckBoxState::Unchecked;
+		ChildrenState.bAnyChildIsChecked = ChildrenState.bAllChildrenAreChecked = CheckedState == ECheckBoxState::Checked;
 		bShouldMigratePackage = bInShouldMigratePackage;
 	}
+
+	return ChildrenState;
 }
 
 BEGIN_SLATE_FUNCTION_BUILD_OPTIMIZATION
@@ -227,8 +237,6 @@ TSharedRef<ITableRow> SPackageReportDialog::GenerateTreeRow( TSharedPtr<FPackage
 				SNew(SCheckBox)
 				.OnCheckStateChanged(this, &SPackageReportDialog::CheckBoxStateChanged, TreeItem, OwnerTable)
 				.IsChecked(this, &SPackageReportDialog::GetEnabledCheckState, TreeItem)
-				.IsEnabled(TreeItem.Get()->Parent == nullptr ? true : TreeItem.Get()->Parent->bIsActive)
-
 			]
 			+SHorizontalBox::Slot()
 			.AutoWidth()
@@ -240,40 +248,62 @@ TSharedRef<ITableRow> SPackageReportDialog::GenerateTreeRow( TSharedPtr<FPackage
 			.FillWidth(1.f)
 			[
 				SNew(STextBlock).Text(FText::FromString(TreeItem->NodeName))
-				.ColorAndOpacity(TreeItem.Get()->bIsActive ? FSlateColor::UseForeground() : FSlateColor::UseSubduedForeground())
+				.ColorAndOpacity(FSlateColor::UseForeground())
 			]
 		];
 }
 
 ECheckBoxState SPackageReportDialog::GetEnabledCheckState(TSharedPtr<FPackageReportNode> TreeItem) const
 {
-	return TreeItem.Get()->bIsChecked ? ECheckBoxState::Checked : ECheckBoxState::Unchecked;
+	return TreeItem.Get()->CheckedState;
 }
 
-void SPackageReportDialog::SetStateRecursive(TSharedPtr<FPackageReportNode> TreeItem, bool bWasChecked)
+void SPackageReportDialog::SetStateRecursive(TSharedPtr<FPackageReportNode> TreeItem, bool bIsChecked)
 {
-	if (TreeItem.Get() == nullptr) return;
+	if (TreeItem.Get() == nullptr)
+	{
+		return;
+	}
+
+	TreeItem.Get()->CheckedState = bIsChecked ? ECheckBoxState::Checked : ECheckBoxState::Unchecked;
 
-	TreeItem.Get()->bIsActive = bWasChecked && TreeItem.Get()->bIsChecked;
 	if (TreeItem.Get()->bShouldMigratePackage)
 	{
-		*(TreeItem.Get()->bShouldMigratePackage) = TreeItem.Get()->bIsActive;
+		*(TreeItem.Get()->bShouldMigratePackage) = bIsChecked;
 	}
 
 	TArray< TSharedPtr<FPackageReportNode> > Children;
 	GetChildrenForTree(TreeItem, Children);
 	for (int i = 0; i < Children.Num(); i++)
 	{
-		if (Children[i].Get() == nullptr) continue;
+		if (Children[i].Get() == nullptr)
+		{
+			continue;
+		}
 
-		SetStateRecursive(Children[i], TreeItem.Get()->bIsActive);
+		SetStateRecursive(Children[i], bIsChecked);
 	}
 }
 
 void SPackageReportDialog::CheckBoxStateChanged(ECheckBoxState InCheckBoxState, TSharedPtr<FPackageReportNode> TreeItem, TSharedRef<STableViewBase> OwnerTable)
 {
-	TreeItem.Get()->bIsChecked = InCheckBoxState == ECheckBoxState::Checked;
 	SetStateRecursive(TreeItem, InCheckBoxState == ECheckBoxState::Checked);
+
+	FPackageReportNode* CurrentParent = TreeItem->Parent;
+	while (CurrentParent != nullptr)
+	{
+		bool bAnyChildIsChecked = false;
+		bool bAllChildrenAreChecked = true;
+		for (int i = 0; i < CurrentParent->Children.Num(); i++)
+		{
+			bAnyChildIsChecked |= CurrentParent->Children[i]->CheckedState != ECheckBoxState::Unchecked;
+			bAllChildrenAreChecked &= CurrentParent->Children[i]->CheckedState != ECheckBoxState::Unchecked;
+		}
+
+		CurrentParent->CheckedState = bAllChildrenAreChecked ? ECheckBoxState::Checked : (bAnyChildIsChecked ? ECheckBoxState::Undetermined : ECheckBoxState::Unchecked);
+		CurrentParent = CurrentParent->Parent;
+	}
+
 	OwnerTable.Get().RebuildList();
 }
 
diff --git a/Engine/Source/Developer/AssetTools/Private/SPackageReportDialog.h b/Engine/Source/Developer/AssetTools/Private/SPackageReportDialog.h
--- a/Engine/Source/Developer/AssetTools/Private/SPackageReportDialog.h
+++ b/Engine/Source/Developer/AssetTools/Private/SPackageReportDialog.h
@@ -25,9 +25,7 @@ struct FPackageReportNode
 	/** The name of the tree node without the path */
 	FString NodeName; 
 	/** A user-exposed flag determining whether the content of this node and its children should be migrated or not. */
-	bool bIsChecked;
-	/** A flag determining whether this node should be migrated or not. This node is active as long as bIsChecked is true and if all the parent nodes are also checked. */
-	bool bIsActive;
+	ECheckBoxState CheckedState;
 	/** A pointer to an external bool describing whether this node ultimately should be migrated or not. Is only non-null for leaf nodes.*/
 	bool* bShouldMigratePackage;
 	/** If true, this node is a folder instead of a package */
@@ -50,8 +48,13 @@ struct FPackageReportNode
 	void ExpandChildrenRecursively(const TSharedRef<PackageReportTree>& Treeview);
 
 private:
+	struct FChildrenState
+	{
+		bool bAnyChildIsChecked;
+		bool bAllChildrenAreChecked;
+	};
 	/** Helper function for AddPackage. PathElements is the tokenized path delimited by "/" */
-	void AddPackage_Recursive(TArray<FString>& PathElements, bool* bInIsPackageIncluded);
+	FChildrenState AddPackage_Recursive(TArray<FString>& PathElements, bool* bInIsPackageIncluded);
 };
 
 class SPackageReportDialog : public SCompoundWidget
@@ -74,7 +77,7 @@ public:
 
 private:
 	/** Recursively sets the checked/active state of every child of this node in the tree when a checkbox is toggled. */
-	void SetStateRecursive(TSharedPtr<FPackageReportNode> TreeItem, bool bWasChecked);
+	void SetStateRecursive(TSharedPtr<FPackageReportNode> TreeItem, bool bIsChecked);
 
 	/** Callback to check whether a checkbox is checked or not. */
 	ECheckBoxState GetEnabledCheckState(TSharedPtr<FPackageReportNode> TreeItem) const;
