Second pass on SlowTask dialogs for Migrate.
Move LoadAssets to after the UI prompts for the UX benefit.

#rnx

diff --git a/Engine/Source/Developer/AssetTools/Private/AssetTools.cpp b/Engine/Source/Developer/AssetTools/Private/AssetTools.cpp
--- a/Engine/Source/Developer/AssetTools/Private/AssetTools.cpp
+++ b/Engine/Source/Developer/AssetTools/Private/AssetTools.cpp
@@ -3098,35 +3098,6 @@ void UAssetToolsImpl::MigratePackages_ReportConfirmed(TSharedPtr<TArray<ReportPa
 	// Fixing up references requires resaving packages to a temporary location
 	if (!LostPackages.IsEmpty())
 	{
-		// Resolve the packages to migrate to assets
-		TArray<UObject*> SrcObjects;
-		for (const FName& SrcPackage : AllPackageNamesToMove)
-		{
-			UPackage* LoadedPackage = UPackageTools::LoadPackage(SrcPackage.ToString());
-			if (!LoadedPackage)
-			{
-				MigrateLog.Error(FText::Format(LOCTEXT("MigratePackages_FailedToLoadPackage", "Failed to load package {0}"), FText::FromString(SrcPackage.ToString())));
-				bAbort = true;
-				continue;
-			}
-
-			UObject* Asset = LoadedPackage->FindAssetInPackage();
-			if (Asset)
-			{
-				SrcObjects.Add(Asset);
-			}
-			else
-			{
-				MigrateLog.Warning(FText::Format(LOCTEXT("MigratePackages_PackageHasNoAsset", "Package {0} has no asset in it"), FText::FromString(SrcPackage.ToString())));
-			}
-		}
-
-		if (bAbort)
-		{
-			MigrateLog.Notify();
-			return;
-		}
-
 		// Query the user for a folder to migrate assets into. This folder will exist temporary in this project, and is the destination for migration in the target project
 		FString FolderName;
 		bool bIsOkButtonEnabled = true;
@@ -3206,21 +3177,51 @@ void UAssetToolsImpl::MigratePackages_ReportConfirmed(TSharedPtr<TArray<ReportPa
 			return;
 		}
 
-		TArray<UObject*> TempObjects;
-		TMap<UObject*, UObject*> ReplacementMap;
+		// Resolve the packages to migrate to assets
+		TArray<UObject*> SrcObjects;
+		{
+			FScopedSlowTask SlowTask(AllPackageNamesToMove.Num(), LOCTEXT("MigratePackages_Loading", "Loading Packages..."));
+			SlowTask.MakeDialog();
+			for (const FName& SrcPackage : AllPackageNamesToMove)
+			{
+				SlowTask.EnterProgressFrame();
+				UPackage* LoadedPackage = UPackageTools::LoadPackage(SrcPackage.ToString());
+				if (!LoadedPackage)
+				{
+					MigrateLog.Error(FText::Format(LOCTEXT("MigratePackages_FailedToLoadPackage", "Failed to load package {0}"), FText::FromString(SrcPackage.ToString())));
+					bAbort = true;
+					continue;
+				}
 
-		// Copy all specified assets and their dependencies to the destination folder
-		FScopedSlowTask SlowTask( 3, LOCTEXT( "MigratePackages_Consolidate", "Consolidating Assets..." ) );
-		SlowTask.MakeDialog();
+				UObject* Asset = LoadedPackage->FindAssetInPackage();
+				if (Asset)
+				{
+					SrcObjects.Add(Asset);
+				}
+				else
+				{
+					MigrateLog.Warning(FText::Format(LOCTEXT("MigratePackages_PackageHasNoAsset", "Package {0} has no asset in it"), FText::FromString(SrcPackage.ToString())));
+				}
+			}
+		}
+
+		if (bAbort)
+		{
+			MigrateLog.Notify();
+			return;
+		}
 
 		// To handle complex references and assets in different Plugins, we must first duplicate to temp packages
+		TArray<UObject*> TempObjects;
+		TMap<UObject*, UObject*> ReplacementMap;
 		{
-			FScopedSlowTask LoopProgress(SrcObjects.Num());
+			FScopedSlowTask SlowTask(SrcObjects.Num(), LOCTEXT("MigratePackages_Duplicate", "Duplicating Assets..."));
+			SlowTask.MakeDialog();
 
 			TSet<UPackage*> PackagesUserRefusedToFullyLoad;
 			for (int i=0; i<SrcObjects.Num(); ++i)
 			{
-				LoopProgress.EnterProgressFrame();
+				SlowTask.EnterProgressFrame();
 
 				UObject* Object = SrcObjects[i];
 				FString PackageName = Object->GetPackage()->GetName();
@@ -3254,11 +3255,12 @@ void UAssetToolsImpl::MigratePackages_ReportConfirmed(TSharedPtr<TArray<ReportPa
 
 		// Update references between TempObjects (to reference each other)
 		{
-			FScopedSlowTask LoopProgress(TempObjects.Num());
+			FScopedSlowTask SlowTask(TempObjects.Num(), LOCTEXT("MigratePackages_ReplaceReferences", "Replacing References..."));
+			SlowTask.MakeDialog();
 
 			for (int i=0; i<TempObjects.Num(); ++i)
 			{
-				LoopProgress.EnterProgressFrame();
+				SlowTask.EnterProgressFrame();
 				UObject* TempObject = TempObjects[i];
 
 				FArchiveReplaceObjectRef<UObject> ReplaceAr(TempObject, ReplacementMap, EArchiveReplaceObjectFlags::IgnoreOuterRef | EArchiveReplaceObjectFlags::IgnoreArchetypeRef);
@@ -3267,12 +3269,13 @@ void UAssetToolsImpl::MigratePackages_ReportConfirmed(TSharedPtr<TArray<ReportPa
 
 		// Save fixed up packages to the migrated folder, and update the set of files to copy to be those migrated packages
 		{
-			FScopedSlowTask LoopProgress(TempObjects.Num());
+			FScopedSlowTask SlowTask(TempObjects.Num(), LOCTEXT("MigratePackages_ReplaceReferences", "Replacing References..."));
+			SlowTask.MakeDialog();
 
 			TSet<FName> NewPackageNamesToMove;
 			for (int i=0; i<TempObjects.Num(); ++i)
 			{
-				LoopProgress.EnterProgressFrame();
+				SlowTask.EnterProgressFrame();
 				UObject* TempObject = TempObjects[i];
 
 				// Calculate the file path to the new, migrated package
@@ -3292,21 +3295,18 @@ void UAssetToolsImpl::MigratePackages_ReportConfirmed(TSharedPtr<TArray<ReportPa
 
 	bool bUserCanceled = false;
 
-	// Copy all specified assets and their dependencies to the destination folder
-	FScopedSlowTask SlowTask( 2, LOCTEXT( "MigratePackages_CopyingFiles", "Copying Files..." ) );
-	SlowTask.MakeDialog();
-
 	EAppReturnType::Type LastResponse = EAppReturnType::Yes;
 	TArray<FString> SuccessfullyCopiedFiles;
 	TArray<FString> SuccessfullyCopiedPackages;
 	FString CopyErrors;
 
-	SlowTask.EnterProgressFrame();
 	{
-		FScopedSlowTask LoopProgress(AllPackageNamesToMove.Num());
+		FScopedSlowTask SlowTask(AllPackageNamesToMove.Num(), LOCTEXT("MigratePackages_CopyingFiles", "Copying Files..."));
+		SlowTask.MakeDialog();
+
 		for ( const FName& PackageNameToMove : AllPackageNamesToMove )
 		{
-			LoopProgress.EnterProgressFrame();
+			SlowTask.EnterProgressFrame();
 
 			const FString& PackageName = PackageNameToMove.ToString();
 			FString SrcFilename;
@@ -3403,9 +3403,14 @@ void UAssetToolsImpl::MigratePackages_ReportConfirmed(TSharedPtr<TArray<ReportPa
 	// If we are consolidating lost packages, we are copying temporary packages, so clean them up.
 	if (!LostPackages.IsEmpty())
 	{
+		FScopedSlowTask SlowTask(AllPackageNamesToMove.Num(), LOCTEXT("MigratePackages_CleaningUp", "Cleaning Up..."));
+		SlowTask.MakeDialog();
+
 		TArray<FAssetData> AssetsToDelete;
 		for (const FName& PackageNameToMove : AllPackageNamesToMove)
 		{
+			SlowTask.EnterProgressFrame();
+
 			UPackage* Package = UPackageTools::LoadPackage(PackageNameToMove.ToString());
 			if (Package)
 			{
@@ -3424,7 +3429,6 @@ void UAssetToolsImpl::MigratePackages_ReportConfirmed(TSharedPtr<TArray<ReportPa
 	}
 
 	FString SourceControlErrors;
-	SlowTask.EnterProgressFrame();
 
 	if ( !bUserCanceled && SuccessfullyCopiedFiles.Num() > 0 )
 	{
@@ -3436,11 +3440,12 @@ void UAssetToolsImpl::MigratePackages_ReportConfirmed(TSharedPtr<TArray<ReportPa
 				ISourceControlProvider& SourceControlProvider = ISourceControlModule::Get().GetProvider();
 				if(SourceControlProvider.Execute(ISourceControlOperation::Create<FMarkForAdd>(), SuccessfullyCopiedFiles) == ECommandResult::Failed)
 				{
-					FScopedSlowTask LoopProgress(SuccessfullyCopiedFiles.Num());
+					FScopedSlowTask SlowTask(SuccessfullyCopiedFiles.Num(), LOCTEXT("MigratePackages_AddToSourceControl", "Adding To Source Control..."));
+					SlowTask.MakeDialog();
 
 					for(auto FileIt(SuccessfullyCopiedFiles.CreateConstIterator()); FileIt; FileIt++)
 					{
-						LoopProgress.EnterProgressFrame();
+						SlowTask.EnterProgressFrame();
 						if(!SourceControlProvider.GetState(*FileIt, EStateCacheUsage::Use)->IsAdded())
 						{
 							SourceControlErrors += FText::Format(LOCTEXT("MigratePackages_SourceControlError", "{0} could not be added to source control"), FText::FromString(*FileIt)).ToString();
