Detect broken Blueprint assets and abort Migration if they are included.

The condition here mimics the condition we use when we advise caution about opening the Blueprint Editor on the same assets.

Use Force Delete to null out references and unload assets so we can delete the package files that Migrate generated.

#rnx

diff --git a/Engine/Source/Developer/AssetTools/Private/AssetTools.cpp b/Engine/Source/Developer/AssetTools/Private/AssetTools.cpp
--- a/Engine/Source/Developer/AssetTools/Private/AssetTools.cpp
+++ b/Engine/Source/Developer/AssetTools/Private/AssetTools.cpp
@@ -3193,15 +3193,23 @@ void UAssetToolsImpl::MigratePackages_ReportConfirmed(TSharedPtr<TArray<ReportPa
 					continue;
 				}
 
-				UObject* Asset = LoadedPackage->FindAssetInPackage();
-				if (Asset)
+				UObject* Object = LoadedPackage->FindAssetInPackage();
+				if (Object)
 				{
-					SrcObjects.Add(Asset);
+					SrcObjects.Add(Object);
 				}
 				else
 				{
 					MigrateLog.Warning(FText::Format(LOCTEXT("MigratePackages_PackageHasNoAsset", "Package {0} has no asset in it"), FText::FromString(SrcPackage.ToString())));
 				}
+
+				UBlueprint* Blueprint = Cast<UBlueprint>(Object);
+				if (Blueprint && (!Blueprint->SkeletonGeneratedClass || !Blueprint->GeneratedClass))
+				{
+					MigrateLog.Error(FText::Format(LOCTEXT("MigratePackages_InvalidBluepriint", "Blueprint is invalid in package {0}"), FText::FromString(SrcPackage.ToString())));
+					bAbort = true;
+					continue;
+				}
 			}
 		}
 
@@ -3406,7 +3414,7 @@ void UAssetToolsImpl::MigratePackages_ReportConfirmed(TSharedPtr<TArray<ReportPa
 		FScopedSlowTask SlowTask(AllPackageNamesToMove.Num(), LOCTEXT("MigratePackages_CleaningUp", "Cleaning Up..."));
 		SlowTask.MakeDialog();
 
-		TArray<FAssetData> AssetsToDelete;
+		TArray<UObject*> ObjectsToDelete;
 		for (const FName& PackageNameToMove : AllPackageNamesToMove)
 		{
 			SlowTask.EnterProgressFrame();
@@ -3414,12 +3422,11 @@ void UAssetToolsImpl::MigratePackages_ReportConfirmed(TSharedPtr<TArray<ReportPa
 			UPackage* Package = UPackageTools::LoadPackage(PackageNameToMove.ToString());
 			if (Package)
 			{
-				FAssetData AssetData(Package);
-				AssetsToDelete.Add(AssetData);
+				ObjectsToDelete.Add(Package);
 			}
 		}
 
-		ObjectTools::DeleteAssets(AssetsToDelete, /*bShowConfirmation=*/false);
+		ObjectTools::ForceDeleteObjects(ObjectsToDelete, /*bShowConfirmation=*/false);
 
 		if (!IFileManager::Get().DeleteDirectory(*SrcDiskFolderFilename, /*RequireExists =*/ false, /*Tree =*/ true))
 		{
