Fixup Redirectors: Fix GCed Empty Packages crash (fix one other use case in FAssetRenameManager that I previously missed)

#rnx

diff --git a/Engine/Source/Developer/AssetTools/Private/AssetFixUpRedirectors.cpp b/Engine/Source/Developer/AssetTools/Private/AssetFixUpRedirectors.cpp
--- a/Engine/Source/Developer/AssetTools/Private/AssetFixUpRedirectors.cpp
+++ b/Engine/Source/Developer/AssetTools/Private/AssetFixUpRedirectors.cpp
@@ -23,6 +23,7 @@
 #include "Logging/MessageLog.h"
 #include "AssetTools.h"
 #include "Engine/Blueprint.h"
+#include "SourceControlHelpers.h"
 
 #define LOCTEXT_NAMESPACE "AssetFixUpRedirectors"
 
@@ -362,11 +363,14 @@ void FAssetFixUpRedirectors::SaveReferencingPackages(const TArray<UPackage*>& Re
 {
 	if ( ReferencingPackagesToSave.Num() > 0 )
 	{
+		// Get the list of filenames before calling save because some of the saved packages can get GCed if they are empty packages
+		const TArray<FString> Filenames = USourceControlHelpers::PackageFilenames(ReferencingPackagesToSave);
+
 		const bool bCheckDirty = false;
 		const bool bPromptToSave = false;
 		FEditorFileUtils::PromptForCheckoutAndSave(ReferencingPackagesToSave, bCheckDirty, bPromptToSave, &OutFailedToSave);
 
-		ISourceControlModule::Get().QueueStatusUpdate(ReferencingPackagesToSave);
+		ISourceControlModule::Get().QueueStatusUpdate(Filenames);
 	}
 }
 
diff --git a/Engine/Source/Developer/AssetTools/Private/AssetRenameManager.cpp b/Engine/Source/Developer/AssetTools/Private/AssetRenameManager.cpp
--- a/Engine/Source/Developer/AssetTools/Private/AssetRenameManager.cpp
+++ b/Engine/Source/Developer/AssetTools/Private/AssetRenameManager.cpp
@@ -51,8 +51,6 @@
 #include "GameMapsSettings.h"
 #include "AssetToolsModule.h"
 #include "IAssetTools.h"
-#include "Algo/Transform.h"
-
 
 #define LOCTEXT_NAMESPACE "AssetRenameManager"
 
@@ -1570,9 +1568,7 @@ void FAssetRenameManager::PerformAssetRename(TArray<FAssetRenameDataWithReferenc
 		const bool bAlreadyCheckedOut = true;
 
 		// Get the list of filenames before calling save because some of the saved packages can get GCed if they are empty packages
-		TArray<FString> Filenames;
-		Filenames.Reserve(PackagesToSave.Num());
-		Algo::Transform(PackagesToSave, Filenames, [](UPackage* InPackage) { return USourceControlHelpers::PackageFilename(InPackage); });
+		const TArray<FString> Filenames = USourceControlHelpers::PackageFilenames(PackagesToSave);
 
 		FEditorFileUtils::PromptForCheckoutAndSave(PackagesToSave, bCheckDirty, bPromptToSave, nullptr, bAlreadyCheckedOut);
 
@@ -1649,11 +1645,14 @@ void FAssetRenameManager::SaveReferencingPackages(const TArray<UPackage*>& Refer
 {
 	if (ReferencingPackagesToSave.Num() > 0)
 	{
+		// Get the list of filenames before calling save because some of the saved packages can get GCed if they are empty packages
+		const TArray<FString> Filenames = USourceControlHelpers::PackageFilenames(ReferencingPackagesToSave);
+
 		const bool bCheckDirty = false;
 		const bool bPromptToSave = false;
 		FEditorFileUtils::PromptForCheckoutAndSave(ReferencingPackagesToSave, bCheckDirty, bPromptToSave);
 
-		ISourceControlModule::Get().QueueStatusUpdate(ReferencingPackagesToSave);
+		ISourceControlModule::Get().QueueStatusUpdate(Filenames);
 	}
 }
 
