Editor: Added a fixup redirectors option that *doesn't* delete the redirectors afterwards, which can be useful for cleaning up a project while others are still actively working and might have created new references to the renamed object

diff --git a/Engine/Source/Developer/AssetTools/Private/AssetFixUpRedirectors.cpp b/Engine/Source/Developer/AssetTools/Private/AssetFixUpRedirectors.cpp
--- a/Engine/Source/Developer/AssetTools/Private/AssetFixUpRedirectors.cpp
+++ b/Engine/Source/Developer/AssetTools/Private/AssetFixUpRedirectors.cpp
@@ -44,7 +44,7 @@ struct FRedirectorRefs
 };
 
 
-void FAssetFixUpRedirectors::FixupReferencers(const TArray<UObjectRedirector*>& Objects, const bool bCheckoutDialogPrompt) const
+void FAssetFixUpRedirectors::FixupReferencers(const TArray<UObjectRedirector*>& Objects, const bool bCheckoutDialogPrompt, ERedirectFixupMode FixupMode) const
 {
 	// Transform array into TWeakObjectPtr array
 	TArray<TWeakObjectPtr<UObjectRedirector>> ObjectWeakPtrs;
@@ -61,18 +61,18 @@ void FAssetFixUpRedirectors::FixupReferencers(const TArray<UObjectRedirector*>&
 		{
 			// Open a dialog asking the user to wait while assets are being discovered
 			SDiscoveringAssetsDialog::OpenDiscoveringAssetsDialog(
-				SDiscoveringAssetsDialog::FOnAssetsDiscovered::CreateSP(this, &FAssetFixUpRedirectors::ExecuteFixUp, ObjectWeakPtrs, bCheckoutDialogPrompt)
+				SDiscoveringAssetsDialog::FOnAssetsDiscovered::CreateSP(this, &FAssetFixUpRedirectors::ExecuteFixUp, ObjectWeakPtrs, bCheckoutDialogPrompt, FixupMode)
 				);
 		}
 		else
 		{
 			// No need to wait, attempt to fix references now.
-			ExecuteFixUp(ObjectWeakPtrs, bCheckoutDialogPrompt);
+			ExecuteFixUp(ObjectWeakPtrs, bCheckoutDialogPrompt, FixupMode);
 		}
 	}
 }
 
-void FAssetFixUpRedirectors::ExecuteFixUp(TArray<TWeakObjectPtr<UObjectRedirector>> Objects, const bool bCheckoutDialogPrompt) const
+void FAssetFixUpRedirectors::ExecuteFixUp(TArray<TWeakObjectPtr<UObjectRedirector>> Objects, const bool bCheckoutDialogPrompt, ERedirectFixupMode FixupMode) const
 {
 	TGuardValue<bool> Guard(bIsFixupReferencersInProgress, true);
 
@@ -133,8 +133,11 @@ void FAssetFixUpRedirectors::ExecuteFixUp(TArray<TWeakObjectPtr<UObjectRedirecto
 				// Wait for package referencers to be updated
 				UpdateAssetReferencers(RedirectorRefsList);
 
-				// Delete any redirectors that are no longer referenced
-				DeleteRedirectors(RedirectorRefsList, FailedToSave);
+				if (FixupMode == ERedirectFixupMode::DeleteFixedUpRedirectors)
+				{
+					// Delete any redirectors that are no longer referenced
+					DeleteRedirectors(RedirectorRefsList, FailedToSave);
+				}
 
 				// Finally, report any failures that happened during the rename
 				ReportFailures(RedirectorRefsList);
diff --git a/Engine/Source/Developer/AssetTools/Private/AssetFixUpRedirectors.h b/Engine/Source/Developer/AssetTools/Private/AssetFixUpRedirectors.h
--- a/Engine/Source/Developer/AssetTools/Private/AssetFixUpRedirectors.h
+++ b/Engine/Source/Developer/AssetTools/Private/AssetFixUpRedirectors.h
@@ -5,6 +5,7 @@
 #include "CoreMinimal.h"
 
 struct FRedirectorRefs;
+enum class ERedirectFixupMode;
 
 class FAssetFixUpRedirectors : public TSharedFromThis<FAssetFixUpRedirectors>
 {
@@ -13,7 +14,7 @@ public:
 	 * Fix up references to the specified redirectors.
 	 * @param bCheckoutDialogPrompt indicates whether to prompt the user with files checkout dialog or silently attempt to checkout all necessary files.
 	 */
-	void FixupReferencers(const TArray<UObjectRedirector*>& Objects, bool bCheckoutDialogPrompt = true) const;
+	void FixupReferencers(const TArray<UObjectRedirector*>& Objects, bool bCheckoutDialogPrompt, ERedirectFixupMode FixupMode) const;
 
 	/** Returns whether redirectors are being fixed up. */
 	bool IsFixupReferencersInProgress() const { return bIsFixupReferencersInProgress; }
@@ -21,7 +22,7 @@ public:
 private:
 
 	/** The core code of the fixup operation */
-	void ExecuteFixUp(TArray<TWeakObjectPtr<UObjectRedirector>> Objects, bool bCheckoutDialogPrompt) const;
+	void ExecuteFixUp(TArray<TWeakObjectPtr<UObjectRedirector>> Objects, bool bCheckoutDialogPrompt, ERedirectFixupMode FixupMode) const;
 
 	/** Fills out the Referencing packages for all the redirectors described in AssetsToPopulate */
 	void PopulateRedirectorReferencers(TArray<FRedirectorRefs>& RedirectorsToPopulate) const;
diff --git a/Engine/Source/Developer/AssetTools/Private/AssetTools.cpp b/Engine/Source/Developer/AssetTools/Private/AssetTools.cpp
--- a/Engine/Source/Developer/AssetTools/Private/AssetTools.cpp
+++ b/Engine/Source/Developer/AssetTools/Private/AssetTools.cpp
@@ -3631,9 +3631,9 @@ void UAssetToolsImpl::RecursiveGetDependenciesAdvanced(const FName& PackageName,
 	}
 }
 
-void UAssetToolsImpl::FixupReferencers(const TArray<UObjectRedirector*>& Objects, bool bCheckoutDialogPrompt) const
+void UAssetToolsImpl::FixupReferencers(const TArray<UObjectRedirector*>& Objects, bool bCheckoutDialogPrompt, ERedirectFixupMode FixupMode) const
 {
-	AssetFixUpRedirectors->FixupReferencers(Objects, bCheckoutDialogPrompt);
+	AssetFixUpRedirectors->FixupReferencers(Objects, bCheckoutDialogPrompt, FixupMode);
 }
 
 bool UAssetToolsImpl::IsFixupReferencersInProgress() const
diff --git a/Engine/Source/Developer/AssetTools/Private/AssetTools.h b/Engine/Source/Developer/AssetTools/Private/AssetTools.h
--- a/Engine/Source/Developer/AssetTools/Private/AssetTools.h
+++ b/Engine/Source/Developer/AssetTools/Private/AssetTools.h
@@ -105,7 +105,7 @@ public:
 	virtual bool CreateDiffProcess(const FString& DiffCommand, const FString& OldTextFilename, const FString& NewTextFilename, const FString& DiffArgs = FString("")) const override;
 	virtual void MigratePackages(const TArray<FName>& PackageNamesToMigrate) const override;
 	virtual void BeginAdvancedCopyPackages(const TArray<FName>& InputNamesToCopy, const FString& TargetPath) const override;
-	virtual void FixupReferencers(const TArray<UObjectRedirector*>& Objects, bool bCheckoutDialogPrompt = true) const override;
+	virtual void FixupReferencers(const TArray<UObjectRedirector*>& Objects, bool bCheckoutDialogPrompt = true, ERedirectFixupMode FixupMode = ERedirectFixupMode::DeleteFixedUpRedirectors) const override;
 	virtual bool IsFixupReferencersInProgress() const override;
 	virtual FAssetPostRenameEvent& OnAssetPostRename() override { return AssetRenameManager->OnAssetPostRenameEvent(); }
 	virtual void ExpandDirectories(const TArray<FString>& Files, const FString& DestinationPath, TArray<TPair<FString, FString>>& FilesAndDestinations) const override;
diff --git a/Engine/Source/Developer/AssetTools/Private/AssetTypeActions/AssetTypeActions_Redirector.cpp b/Engine/Source/Developer/AssetTools/Private/AssetTypeActions/AssetTypeActions_Redirector.cpp
--- a/Engine/Source/Developer/AssetTools/Private/AssetTypeActions/AssetTypeActions_Redirector.cpp
+++ b/Engine/Source/Developer/AssetTools/Private/AssetTypeActions/AssetTypeActions_Redirector.cpp
@@ -27,10 +27,21 @@ void FAssetTypeActions_Redirector::GetActions(const TArray<UObject*>& InObjects,
 		LOCTEXT("Redirector_FixUpTooltip", "Finds referencers to selected redirectors and resaves them if possible, then deletes any redirectors that had all their referencers fixed."),
 		FSlateIcon(),
 		FUIAction(
-			FExecuteAction::CreateSP( this, &FAssetTypeActions_Redirector::ExecuteFixUp, Redirectors ),
+			FExecuteAction::CreateSP( this, &FAssetTypeActions_Redirector::ExecuteFixUp, Redirectors, /*bDeleteAssets=*/ true ),
 			FCanExecuteAction()
 			)
 		);
+
+	Section.AddMenuEntry(
+		"Redirector_FixUp_KeepRedirectors",
+		LOCTEXT("Redirector_FixUp_KeepingRedirector", "Fix Up (keeping redirector)"),
+		LOCTEXT("Redirector_FixUp_KeepingRedirectorTooltip", "Finds referencers to selected redirectors and resaves them if possible."),
+		FSlateIcon(),
+		FUIAction(
+			FExecuteAction::CreateSP( this, &FAssetTypeActions_Redirector::ExecuteFixUp, Redirectors, /*bDeleteAssets=*/ false ),
+			FCanExecuteAction()
+		)
+	);
 }
 
 bool FAssetTypeActions_Redirector::AssetsActivatedOverride( const TArray<UObject*>& InObjects, EAssetTypeActivationMethod::Type ActivationType )
@@ -76,7 +87,7 @@ void FAssetTypeActions_Redirector::ExecuteFindTarget(TArray<TWeakObjectPtr<UObje
 	}
 }
 
-void FAssetTypeActions_Redirector::ExecuteFixUp(TArray<TWeakObjectPtr<UObjectRedirector>> Objects)
+void FAssetTypeActions_Redirector::ExecuteFixUp(TArray<TWeakObjectPtr<UObjectRedirector>> Objects, bool bDeleteAssets)
 {
 	// This will fix references to selected redirectors, except in the following cases:
 	// Redirectors referenced by unloaded maps will not be fixed up, but any references to it that can be fixed up will
@@ -92,7 +103,7 @@ void FAssetTypeActions_Redirector::ExecuteFixUp(TArray<TWeakObjectPtr<UObjectRed
 			ObjectRedirectors.Add(Object.Get());
 		}
 
-		FAssetTools::Get().FixupReferencers(ObjectRedirectors);
+		FAssetTools::Get().FixupReferencers(ObjectRedirectors, /*bCheckoutDialogPrompt=*/ true, bDeleteAssets ? ERedirectFixupMode::DeleteFixedUpRedirectors : ERedirectFixupMode::LeaveFixedUpRedirectors);
 	}
 }
 
diff --git a/Engine/Source/Developer/AssetTools/Private/AssetTypeActions/AssetTypeActions_Redirector.h b/Engine/Source/Developer/AssetTools/Private/AssetTypeActions/AssetTypeActions_Redirector.h
--- a/Engine/Source/Developer/AssetTools/Private/AssetTypeActions/AssetTypeActions_Redirector.h
+++ b/Engine/Source/Developer/AssetTools/Private/AssetTypeActions/AssetTypeActions_Redirector.h
@@ -25,7 +25,7 @@ private:
 	void ExecuteFindTarget(TArray<TWeakObjectPtr<UObjectRedirector>> Objects);
 
 	/** Handler for when FixUp is selected */
-	void ExecuteFixUp(TArray<TWeakObjectPtr<UObjectRedirector>> Objects);
+	void ExecuteFixUp(TArray<TWeakObjectPtr<UObjectRedirector>> Objects, bool bDeleteAssets);
 
 	/** Syncs the content browser to the destination objects for all the supplied redirectors */
 	void FindTargets(const TArray<UObjectRedirector*>& Redirectors) const;
diff --git a/Engine/Source/Developer/AssetTools/Public/IAssetTools.h b/Engine/Source/Developer/AssetTools/Public/IAssetTools.h
--- a/Engine/Source/Developer/AssetTools/Public/IAssetTools.h
+++ b/Engine/Source/Developer/AssetTools/Public/IAssetTools.h
@@ -36,6 +36,16 @@ enum class EAssetRenameResult : uint8
 	Pending,
 };
 
+UENUM()
+enum class ERedirectFixupMode
+{
+	// Remove fully fixed up redirectors after the fixup process is done
+	DeleteFixedUpRedirectors,
+
+	// Leave the redirectors around even if no longer locally referenced
+	LeaveFixedUpRedirectors
+};
+
 USTRUCT(BlueprintType)
 struct FAssetRenameData
 {
@@ -395,7 +405,7 @@ public:
 	 * Fix up references to the specified redirectors.
 	 * @param bCheckoutDialogPrompt indicates whether to prompt the user with files checkout dialog or silently attempt to checkout all necessary files.
 	 */
-	virtual void FixupReferencers(const TArray<UObjectRedirector*>& Objects, bool bCheckoutDialogPrompt = true) const = 0;
+	virtual void FixupReferencers(const TArray<UObjectRedirector*>& Objects, bool bCheckoutDialogPrompt = true, ERedirectFixupMode FixupMode = ERedirectFixupMode::DeleteFixedUpRedirectors) const = 0;
 
 	/** Returns whether redirectors are being fixed up. */
 	virtual bool IsFixupReferencersInProgress() const = 0;
