fix multiple asset importers registered at same priority creates random inconsistency
ImgMedia and normal texture import both registered for EXR
would randomly import one or the other on different runs of Editor
adjust priorities so they are not ambiguous
also add detection and logging to prevent this

diff --git a/Engine/Plugins/Media/ImgMedia/Source/ImgMediaEditor/Private/Factories/ImgMediaSourceFactory.cpp b/Engine/Plugins/Media/ImgMedia/Source/ImgMediaEditor/Private/Factories/ImgMediaSourceFactory.cpp
--- a/Engine/Plugins/Media/ImgMedia/Source/ImgMediaEditor/Private/Factories/ImgMediaSourceFactory.cpp
+++ b/Engine/Plugins/Media/ImgMedia/Source/ImgMediaEditor/Private/Factories/ImgMediaSourceFactory.cpp
@@ -11,10 +11,13 @@
 UImgMediaSourceFactory::UImgMediaSourceFactory(const FObjectInitializer& ObjectInitializer)
 	: Super(ObjectInitializer)
 {
-	Formats.Add(TEXT("exr;EXR Image Sequence"));
+	Formats.Add(TEXT("exr;EXR ImgMedia Image Sequence"));
 
 	SupportedClass = UImgMediaSource::StaticClass();
 	bEditorImport = true;
+	
+	// Required to allow texture factory to take priority when importing new image files
+	ImportPriority = DefaultImportPriority - 1;
 }
 
 
diff --git a/Engine/Source/Developer/AssetTools/Private/AssetTools.cpp b/Engine/Source/Developer/AssetTools/Private/AssetTools.cpp
--- a/Engine/Source/Developer/AssetTools/Private/AssetTools.cpp
+++ b/Engine/Source/Developer/AssetTools/Private/AssetTools.cpp
@@ -2023,20 +2023,44 @@ TArray<UObject*> UAssetToolsImpl::ImportAssetsInternal(const TArray<FString>& Fi
 		{
 			const TArray<UFactory*>& Factories = *FactoriesPtr;
 
-			// Handle the potential of multiple factories being found
 			if(Factories.Num() > 0)
 			{
-				Factory = Factories[0];
+				// Handle the potential of multiple factories being found
+				//	Factories was previously sorted by ImportPriority
+				//  and filtered by file extension
+				//  but FactoryCanImport has not been checked yet
 
 				for(auto FactoryIt = Factories.CreateConstIterator(); FactoryIt; ++FactoryIt)
 				{
 					UFactory* TestFactory = *FactoryIt;
 					if(TestFactory->FactoryCanImport(Filename))
 					{
+						if ( Factory != nullptr )
+						{
+							// we found a second factory for this type
+							if ( Factory->ImportPriority == TestFactory->ImportPriority )
+							{
+								UE_LOG(LogAssetTools, Warning, TEXT("Two factories registered with same priority : %s and %s"), *Factory->GetName(), *TestFactory->GetName() );
+							}
+							break;
+						}
+
 						Factory = TestFactory;
-						break;
+						//found one, continue so we can check for multiple importers 
+						//break;
 					}
 				}
+
+				if ( Factory == nullptr )
+				{
+					// no factories passed FactoryCanImport()
+					//  this seems wrong, but it preserves old behavior
+					// ??
+					
+					Factory = Factories[0];
+					
+					UE_LOG(LogAssetTools, Warning, TEXT("Factory did not pass FactoryCanImport, trying anyway : %s"), *Factory->GetName());
+				}
 			}
 		}
 		else
@@ -2432,33 +2456,54 @@ void UAssetToolsImpl::ExportAssetsInternal(const TArray<UObject*>& ObjectsToExpo
 		}
 
 		// If FBX is listed, make that the most preferred option
-		const FString PreferredExtension = TEXT("FBX");
-		int32 ExtIndex = PreferredExtensions.Find(PreferredExtension);
-		if (ExtIndex > 0)
+		//   also prioritize PNG and EXR ahead of other image formats
+		// @todo provide a general purpose way of prioritizing export formats, don't hard-code FBX here
+		const TCHAR * TopPriorityExtensions[] =
 		{
-			PreferredExtensions.RemoveAt(ExtIndex);
-			PreferredExtensions.Insert(PreferredExtension, 0);
+			TEXT("FBX"),
+			TEXT("PNG"),
+			TEXT("EXR")
+		};
+		const int TopPriorityExtensionsCount = 3;
+
+		for(int TopPriorityExtensionsIndex=0;TopPriorityExtensionsIndex<TopPriorityExtensionsCount;TopPriorityExtensionsIndex++)
+		{
+			FString ThisExtension = TopPriorityExtensions[TopPriorityExtensionsIndex];
+			// FString.Find is case insensitive
+			int32 ExtIndex = PreferredExtensions.Find(ThisExtension);
+			if (ExtIndex != INDEX_NONE)
+			{
+				PreferredExtensions.RemoveAt(ExtIndex);
+				PreferredExtensions.Insert(ThisExtension, 0);
+				// only do for first one found:
+				break;
+			}
 		}
+
+		// if there are multiple exporters, we arbitrarily choose [0] to go first
+		//	@todo sort by alpha or priority or something to make this order consistent
 		FString FirstExtension = PreferredExtensions[0];
 
-		// If FBX is listed, make that the first option here too, then compile them all into one string
+		// If TopPriorityExtension is listed, make that the first option here too, then compile them all into one string
 		check(AllFileTypes.Num() == AllExtensions.Num())
-			for (ExtIndex = 1; ExtIndex < AllFileTypes.Num(); ++ExtIndex)
+
+		for (int ExtIndex = 1; ExtIndex < AllFileTypes.Num(); ++ExtIndex)
+		{
+			const FString FileType = AllFileTypes[ExtIndex];
+			if (FileType.Contains(FirstExtension))
 			{
-				const FString FileType = AllFileTypes[ExtIndex];
-				if (FileType.Contains(PreferredExtension))
-				{
-					AllFileTypes.RemoveAt(ExtIndex);
-					AllFileTypes.Insert(FileType, 0);
+				AllFileTypes.RemoveAt(ExtIndex);
+				AllFileTypes.Insert(FileType, 0);
 
-					const FString Extension = AllExtensions[ExtIndex];
-					AllExtensions.RemoveAt(ExtIndex);
-					AllExtensions.Insert(Extension, 0);
-				}
+				const FString Extension = AllExtensions[ExtIndex];
+				AllExtensions.RemoveAt(ExtIndex);
+				AllExtensions.Insert(Extension, 0);
 			}
+		}
+
 		FString FileTypes;
 		FString Extensions;
-		for (ExtIndex = 0; ExtIndex < AllFileTypes.Num(); ++ExtIndex)
+		for (int ExtIndex = 0; ExtIndex < AllFileTypes.Num(); ++ExtIndex)
 		{
 			if (FileTypes.Len())
 			{
diff --git a/Engine/Source/Editor/UnrealEd/Private/Factories/EditorFactories.cpp b/Engine/Source/Editor/UnrealEd/Private/Factories/EditorFactories.cpp
--- a/Engine/Source/Editor/UnrealEd/Private/Factories/EditorFactories.cpp
+++ b/Engine/Source/Editor/UnrealEd/Private/Factories/EditorFactories.cpp
@@ -5592,10 +5592,12 @@ UReimportTextureFactory.
 UReimportTextureFactory::UReimportTextureFactory(const FObjectInitializer& ObjectInitializer)
 	: Super(ObjectInitializer)
 {
-
 	SupportedClass = UTexture::StaticClass();
 
 	bCreateNew = false;
+	
+	// Required to allow texture factory to take priority when importing new image files
+	ImportPriority = DefaultImportPriority - 1;
 }
 
 UTexture2D* UReimportTextureFactory::CreateTexture2D( UObject* InParent, FName Name, EObjectFlags Flags )
