Update create and duplicate behavior in Content Browser for Public/Private state
- API to control default behavior of whether newly created content is made public or private on creation
- Ensure duplication copies original's Public/Private state

diff --git a/Engine/Source/Developer/AssetTools/Private/AssetTools.cpp b/Engine/Source/Developer/AssetTools/Private/AssetTools.cpp
--- a/Engine/Source/Developer/AssetTools/Private/AssetTools.cpp
+++ b/Engine/Source/Developer/AssetTools/Private/AssetTools.cpp
@@ -185,6 +185,7 @@ UAssetToolsImpl::UAssetToolsImpl(const FObjectInitializer& ObjectInitializer)
 	, AssetClassPermissionList(MakeShared<FNamePermissionList>())
 	, FolderPermissionList(MakeShared<FPathPermissionList>())
 	, WritableFolderPermissionList(MakeShared<FPathPermissionList>())
+	, CreateAssetsAsExternallyReferenceable(true)
 {
 	TArray<FString> SupportedTypesArray;
 	GConfig->GetArray(TEXT("AssetTools"), TEXT("SupportedAssetTypes"), SupportedTypesArray, GEditorIni);
@@ -527,6 +528,9 @@ UObject* UAssetToolsImpl::CreateAsset(const FString& AssetName, const FString& P
 
 	if( NewObj )
 	{
+
+		Pkg->SetIsExternallyReferenceable(CreateAssetsAsExternallyReferenceable);
+
 		// Notify the asset registry
 		FAssetRegistryModule::AssetCreated(NewObj);
 
@@ -689,6 +693,16 @@ UObject* UAssetToolsImpl::PerformDuplicateAsset(const FString& AssetName, const
 	return NewObject;
 }
 
+void UAssetToolsImpl::SetCreateAssetsAsExternallyReferenceable(bool bValue)
+{
+	CreateAssetsAsExternallyReferenceable = bValue;
+}
+
+bool UAssetToolsImpl::GetCreateAssetsAsExternallyReferenceable()
+{
+	return CreateAssetsAsExternallyReferenceable;
+}
+
 void UAssetToolsImpl::GenerateAdvancedCopyDestinations(FAdvancedCopyParams& InParams, const TArray<FName>& InPackageNamesToCopy, const UAdvancedCopyCustomization* CopyCustomization, TMap<FString, FString>& OutPackagesAndDestinations) const
 {
 	IDesktopPlatform* DesktopPlatform = FDesktopPlatformModule::Get();
diff --git a/Engine/Source/Developer/AssetTools/Private/AssetTools.h b/Engine/Source/Developer/AssetTools/Private/AssetTools.h
--- a/Engine/Source/Developer/AssetTools/Private/AssetTools.h
+++ b/Engine/Source/Developer/AssetTools/Private/AssetTools.h
@@ -83,6 +83,8 @@ public:
 	virtual UObject* DuplicateAsset(const FString& AssetName, const FString& PackagePath, UObject* OriginalObject) override;
 	virtual UObject* DuplicateAssetWithDialog(const FString& AssetName, const FString& PackagePath, UObject* OriginalObject) override;
 	virtual UObject* DuplicateAssetWithDialogAndTitle(const FString& AssetName, const FString& PackagePath, UObject* OriginalObject, FText DialogTitle) override;
+	virtual void SetCreateAssetsAsExternallyReferenceable(bool bValue) override;
+	virtual bool GetCreateAssetsAsExternallyReferenceable() override;
 	virtual bool RenameAssets(const TArray<FAssetRenameData>& AssetsAndNames) override;
 	virtual EAssetRenameResult RenameAssetsWithDialog(const TArray<FAssetRenameData>& AssetsAndNames, bool bAutoCheckout = false) override;
 	virtual void FindSoftReferencesToObject(FSoftObjectPath TargetObject, TArray<UObject*>& ReferencingObjects) override;
@@ -220,6 +222,8 @@ private:
 
 	/** List of sub content paths denied for every mount. */
 	TArray<FString> SubContentDenyListPaths;
+
+	bool CreateAssetsAsExternallyReferenceable;
 };
 
 PRAGMA_ENABLE_DEPRECATION_WARNINGS
diff --git a/Engine/Source/Developer/AssetTools/Public/IAssetTools.h b/Engine/Source/Developer/AssetTools/Public/IAssetTools.h
--- a/Engine/Source/Developer/AssetTools/Public/IAssetTools.h
+++ b/Engine/Source/Developer/AssetTools/Public/IAssetTools.h
@@ -253,6 +253,12 @@ public:
 	UFUNCTION(BlueprintCallable, Category="Editor Scripting | Asset Tools")
 	virtual UObject* DuplicateAsset(const FString& AssetName, const FString& PackagePath, UObject* OriginalObject) = 0;
 
+	/** Controls whether or not newly created assets are made externally referneceable or not */
+	virtual void SetCreateAssetsAsExternallyReferenceable(bool bValue) = 0;
+
+	/** Gets whether assets are being made externally referenceable or not */
+	virtual bool GetCreateAssetsAsExternallyReferenceable() = 0;
+
 	/** Renames assets using the specified names. */
 	UFUNCTION(BlueprintCallable, Category = "Editor Scripting | Asset Tools")
 	virtual bool RenameAssets(const TArray<FAssetRenameData>& AssetsAndNames) = 0;
diff --git a/Engine/Source/Editor/UnrealEd/Private/ObjectTools.cpp b/Engine/Source/Editor/UnrealEd/Private/ObjectTools.cpp
--- a/Engine/Source/Editor/UnrealEd/Private/ObjectTools.cpp
+++ b/Engine/Source/Editor/UnrealEd/Private/ObjectTools.cpp
@@ -894,6 +894,10 @@ namespace ObjectTools
 				DupObject->GetOutermost()->SetPackageFlags(PKG_DisallowExport);
 			}
 
+			// Duplicating an asset should respect the reference controls of the original.
+			bool isObjectExternallyReferenceable = Object->GetOutermost()->IsExternallyReferenceable();
+			DupObject->GetOutermost()->SetIsExternallyReferenceable(isObjectExternallyReferenceable);
+
 			// When duplicating a World Composition map, make sure to properly initialize WorldTileInfo
 			if (Object->GetOutermost()->GetWorldTileInfo())
 			{
