Allow separate permission lists for viewing and creating assets

diff --git a/Engine/Source/Developer/AssetTools/Private/AssetTools.cpp b/Engine/Source/Developer/AssetTools/Private/AssetTools.cpp
--- a/Engine/Source/Developer/AssetTools/Private/AssetTools.cpp
+++ b/Engine/Source/Developer/AssetTools/Private/AssetTools.cpp
@@ -181,19 +181,22 @@ UAssetToolsImpl::UAssetToolsImpl(const FObjectInitializer& ObjectInitializer)
 	: Super(ObjectInitializer)
 	, AssetRenameManager(MakeShareable(new FAssetRenameManager))
 	, AssetFixUpRedirectors(MakeShareable(new FAssetFixUpRedirectors))
-	, NextUserCategoryBit(EAssetTypeCategories::FirstUser)
-	, AssetClassPermissionList(MakeShared<FNamePermissionList>())
+	, NextUserCategoryBit(EAssetTypeCategories::FirstUser) 
 	, FolderPermissionList(MakeShared<FPathPermissionList>())
 	, WritableFolderPermissionList(MakeShared<FPathPermissionList>())
 	, CreateAssetsAsExternallyReferenceable(true)
 {
 	TArray<FString> SupportedTypesArray;
 	GConfig->GetArray(TEXT("AssetTools"), TEXT("SupportedAssetTypes"), SupportedTypesArray, GEditorIni);
-	for (const FString& Type : SupportedTypesArray)
+	for (int32 i = 0; i < (int32)EAssetClassAction::AllAssetActions; ++i)
 	{
-		AssetClassPermissionList->AddAllowListItem("AssetToolsConfigFile", *Type);
+		AssetClassPermissionList.Add(MakeShared<FNamePermissionList>());
+		for (const FString& Type : SupportedTypesArray)
+		{
+			AssetClassPermissionList[i]->AddAllowListItem("AssetToolsConfigFile", *Type);
+		}
+		AssetClassPermissionList[i]->OnFilterChanged().AddUObject(this, &UAssetToolsImpl::AssetClassPermissionListChanged, (EAssetClassAction)i);
 	}
-	AssetClassPermissionList->OnFilterChanged().AddUObject(this, &UAssetToolsImpl::AssetClassPermissionListChanged);
 
 	TArray<FString> DenyListedViewPath;
 	GConfig->GetArray(TEXT("AssetTools"), TEXT("DenyListAssetPaths"), DenyListedViewPath, GEditorIni);
@@ -332,7 +335,7 @@ void UAssetToolsImpl::RegisterAssetTypeActions(const TSharedRef<IAssetTypeAction
 	bool bSupported = false;
 	if (const UClass* SupportedClass = NewActions->GetSupportedClass())
 	{
-		bSupported = AssetClassPermissionList->PassesFilter(SupportedClass->GetFName());
+		bSupported = GetAssetClassPermissionList(EAssetClassAction::CreateAsset)->PassesFilter(SupportedClass->GetFName());
 	}
 	else
 	{
@@ -3976,24 +3979,42 @@ TArray<UFactory*> UAssetToolsImpl::GetNewAssetFactories() const
 
 TSharedRef<FNamePermissionList>& UAssetToolsImpl::GetAssetClassPermissionList()
 {
-	return AssetClassPermissionList;
+	return GetAssetClassPermissionList(EAssetClassAction::ViewAsset);
 }
 
-void UAssetToolsImpl::AssetClassPermissionListChanged()
+TSharedRef<FNamePermissionList>& UAssetToolsImpl::GetAssetClassPermissionList(EAssetClassAction AssetClassAction)
 {
-	for (TSharedRef<IAssetTypeActions>& ActionsIt : AssetTypeActionsList)
+	if (AssetClassAction < EAssetClassAction::AllAssetActions)
 	{
-		bool bSupported = false;
-		if (const UClass* SupportedClass = ActionsIt->GetSupportedClass())
-		{
-			bSupported = AssetClassPermissionList->PassesFilter(SupportedClass->GetFName());
-		}
-		else
+		return AssetClassPermissionList[(int32)AssetClassAction];
+	}
+	
+	static TSharedRef<FNamePermissionList> Empty = MakeShared<FNamePermissionList>();
+	return Empty;
+}
+
+void UAssetToolsImpl::AssetClassPermissionListChanged(EAssetClassAction AssetClassAction)
+{
+	switch (AssetClassAction)
+	{
+	case EAssetClassAction::CreateAsset:
+		for (TSharedRef<IAssetTypeActions>& ActionsIt : AssetTypeActionsList)
 		{
-			bSupported = !ActionsIt->GetFilterName().IsNone();
-		}
+			bool bSupported = false;
+			if (const UClass* SupportedClass = ActionsIt->GetSupportedClass())
+			{
+				bSupported = GetAssetClassPermissionList(AssetClassAction)->PassesFilter(SupportedClass->GetFName());
+			}
+			else
+			{
+				bSupported = !ActionsIt->GetFilterName().IsNone();
+			}
 
-		ActionsIt->SetSupported(bSupported);
+			ActionsIt->SetSupported(bSupported);
+		}
+		break;
+	default:
+		break;
 	}
 }
 
diff --git a/Engine/Source/Developer/AssetTools/Private/AssetTools.h b/Engine/Source/Developer/AssetTools/Private/AssetTools.h
--- a/Engine/Source/Developer/AssetTools/Private/AssetTools.h
+++ b/Engine/Source/Developer/AssetTools/Private/AssetTools.h
@@ -124,6 +124,7 @@ public:
 	virtual bool IsAssetClassSupported(const UClass* AssetClass) const override;
 	virtual TArray<UFactory*> GetNewAssetFactories() const override;
 	virtual TSharedRef<FNamePermissionList>& GetAssetClassPermissionList() override;
+	virtual TSharedRef<FNamePermissionList>& GetAssetClassPermissionList(EAssetClassAction AssetClassAction) override;
 	virtual TSharedRef<FPathPermissionList>& GetFolderPermissionList() override;
 	virtual TSharedRef<FPathPermissionList>& GetWritableFolderPermissionList() override;
 	virtual bool AllPassWritableFolderFilter(const TArray<FString>& InPaths) const override;
@@ -184,7 +185,7 @@ private:
 	UObject* PerformDuplicateAsset(const FString& AssetName, const FString& PackagePath, UObject* OriginalObject, bool bWithDialog);
 
 	/** Internal method that performs actions when asset class deny list filter changes */
-	void AssetClassPermissionListChanged();
+	void AssetClassPermissionListChanged(EAssetClassAction AssetClassAction);
 
 	/**
 	 * Add sub content deny list filter for a new mount point
@@ -211,8 +212,8 @@ private:
 	/** The next user category bit to allocate (set to 0 when there are no more bits left) */
 	uint32 NextUserCategoryBit;
 
-	/** Permission list of assets by class name */
-	TSharedRef<FNamePermissionList> AssetClassPermissionList;
+	/** Permission lists of assets by class name, one for each EAssetClassAction */
+	TArray<TSharedRef<FNamePermissionList>> AssetClassPermissionList;
 
 	/** Permission list of folder paths */
 	TSharedRef<FPathPermissionList> FolderPermissionList;
diff --git a/Engine/Source/Developer/AssetTools/Public/IAssetTools.h b/Engine/Source/Developer/AssetTools/Public/IAssetTools.h
--- a/Engine/Source/Developer/AssetTools/Public/IAssetTools.h
+++ b/Engine/Source/Developer/AssetTools/Public/IAssetTools.h
@@ -25,6 +25,16 @@ class UAdvancedCopyCustomization;
 class FNamePermissionList;
 class FPathPermissionList;
 
+UENUM()
+enum class EAssetClassAction : uint8
+{
+	/** Whether an action can be created by an AssetTypeAction */
+	CreateAsset,
+	/** Whether an asset can be viewed in the content browser */
+	ViewAsset,
+	AllAssetActions
+};
+
 UENUM()
 enum class EAssetRenameResult : uint8
 {
@@ -457,8 +467,12 @@ public:
 	/** Find all supported asset factories. */
 	virtual TArray<UFactory*> GetNewAssetFactories() const = 0;
 
-	/** Get asset class permission list for content browser and other systems */
+	/** Get asset class permission list for the ViewAsset action */
+	UE_DEPRECATED(5.1, "Pass in an EAssetClassAction instead of nothing")
 	virtual TSharedRef<FNamePermissionList>& GetAssetClassPermissionList() = 0;
+	
+	/** Get asset class permission list for content browser and other systems */
+	virtual TSharedRef<FNamePermissionList>& GetAssetClassPermissionList(EAssetClassAction AssetClassAction) = 0;
 
 	/** Get folder permission list for content browser and other systems */
 	virtual TSharedRef<FPathPermissionList>& GetFolderPermissionList() = 0;
diff --git a/Engine/Source/Editor/ContentBrowser/Private/SAssetView.cpp b/Engine/Source/Editor/ContentBrowser/Private/SAssetView.cpp
--- a/Engine/Source/Editor/ContentBrowser/Private/SAssetView.cpp
+++ b/Engine/Source/Editor/ContentBrowser/Private/SAssetView.cpp
@@ -334,7 +334,7 @@ void SAssetView::Construct( const FArguments& InArgs )
 	OwningContentBrowser = InArgs._OwningContentBrowser;
 
 	FAssetToolsModule& AssetToolsModule = FModuleManager::LoadModuleChecked<FAssetToolsModule>("AssetTools");
-	AssetClassPermissionList = AssetToolsModule.Get().GetAssetClassPermissionList();
+	AssetClassPermissionList = AssetToolsModule.Get().GetAssetClassPermissionList(EAssetClassAction::ViewAsset);
 	FolderPermissionList = AssetToolsModule.Get().GetFolderPermissionList();
 	WritableFolderPermissionList = AssetToolsModule.Get().GetWritableFolderPermissionList();
 
diff --git a/Engine/Source/Editor/ContentBrowser/Private/SFilterList.cpp b/Engine/Source/Editor/ContentBrowser/Private/SFilterList.cpp
--- a/Engine/Source/Editor/ContentBrowser/Private/SFilterList.cpp
+++ b/Engine/Source/Editor/ContentBrowser/Private/SFilterList.cpp
@@ -1358,7 +1358,7 @@ void SFilterList::PopulateAddFilterMenu(UToolMenu* Menu)
 	};
 	AssetTypeActionsList.Sort( FCompareIAssetTypeActions() );
 
-	TSharedRef<FNamePermissionList> AssetClassPermissionList = AssetToolsModule.Get().GetAssetClassPermissionList();
+	const TSharedRef<FNamePermissionList>& AssetClassPermissionList = AssetToolsModule.Get().GetAssetClassPermissionList(EAssetClassAction::CreateAsset);
 
 	// For every asset type, move it into all the categories it should appear in
 	for (int32 ClassIdx = 0; ClassIdx < AssetTypeActionsList.Num(); ++ClassIdx)
@@ -1649,7 +1649,7 @@ void SFilterList::GetTypeActionsForCategory(EAssetTypeCategories::Type Category,
 	FAssetToolsModule& AssetToolsModule = FModuleManager::LoadModuleChecked<FAssetToolsModule>(TEXT("AssetTools"));
 	TArray<TWeakPtr<IAssetTypeActions>> AssetTypeActionsList;
 	AssetToolsModule.Get().GetAssetTypeActionsList(AssetTypeActionsList);
-	TSharedRef<FNamePermissionList> AssetClassPermissionList = AssetToolsModule.Get().GetAssetClassPermissionList();
+	const TSharedRef<FNamePermissionList>& AssetClassPermissionList = AssetToolsModule.Get().GetAssetClassPermissionList(EAssetClassAction::ViewAsset);
 
 	// Find all asset type actions that match the category
 	for (int32 ClassIdx = 0; ClassIdx < AssetTypeActionsList.Num(); ++ClassIdx)
