Adding ability to batch replace missing skeleton asset.

diff --git a/Engine/Source/Developer/AssetTools/Private/AssetTypeActions/AssetTypeActions_AnimationAsset.cpp b/Engine/Source/Developer/AssetTools/Private/AssetTypeActions/AssetTypeActions_AnimationAsset.cpp
--- a/Engine/Source/Developer/AssetTools/Private/AssetTypeActions/AssetTypeActions_AnimationAsset.cpp
+++ b/Engine/Source/Developer/AssetTools/Private/AssetTypeActions/AssetTypeActions_AnimationAsset.cpp
@@ -48,6 +48,17 @@ void FAssetTypeActions_AnimationAsset::GetActions(const TArray<UObject*>& InObje
 			FCanExecuteAction()
 			)
 		);
+
+	Section.AddMenuEntry(
+		"AnimSequenceBase_ReplaceSkeleton",
+		LOCTEXT("AnimSequenceBase_ReplaceSkeleton", "Replace Skeleton"),
+		LOCTEXT("AnimSequenceBase_ReplaceSkeletonTooltip", "Associate a different skeleton with the selected animation assets in the content browser."),
+		FSlateIcon(FAppStyle::GetAppStyleSetName(), "Persona.RetargetManager"),
+		FUIAction(
+			FExecuteAction::CreateSP( this, &FAssetTypeActions_AnimationAsset::ExecuteReplaceSkeleton, AnimAssets ),
+			FCanExecuteAction()
+			)
+		);
 }
 
 UThumbnailInfo* FAssetTypeActions_AnimationAsset::GetThumbnailInfo(UObject* Asset) const
@@ -217,6 +228,17 @@ void FAssetTypeActions_AnimationAsset::ExecuteFindSkeleton(TArray<TWeakObjectPtr
 	}
 }
 
+void FAssetTypeActions_AnimationAsset::ExecuteReplaceSkeleton(TArray<TWeakObjectPtr<UAnimationAsset>> Objects)
+{
+	TArray<TObjectPtr<UObject>> LoadedAnimAssets;
+	for (const TWeakObjectPtr<UAnimationAsset>& AnimAsset : Objects)
+	{
+		LoadedAnimAssets.Add(AnimAsset.Get());
+	}
+
+	ReplaceMissingSkeleton(LoadedAnimAssets);
+}
+
 bool FAssetTypeActions_AnimationAsset::ReplaceMissingSkeleton(TArray<TObjectPtr<UObject>> InAnimationAssets) const
 {
 	// record anim assets that need skeleton replaced
diff --git a/Engine/Source/Developer/AssetTools/Public/AssetTypeActions/AssetTypeActions_AnimationAsset.h b/Engine/Source/Developer/AssetTools/Public/AssetTypeActions/AssetTypeActions_AnimationAsset.h
--- a/Engine/Source/Developer/AssetTools/Public/AssetTypeActions/AssetTypeActions_AnimationAsset.h
+++ b/Engine/Source/Developer/AssetTools/Public/AssetTypeActions/AssetTypeActions_AnimationAsset.h
@@ -29,6 +29,9 @@ private:
 	/** Handler for when FindSkeleton is selected */
 	void ExecuteFindSkeleton(TArray<TWeakObjectPtr<UAnimationAsset>> Objects);
 
+	/** Handler for when ReplaceSkeleton is selected */
+	void ExecuteReplaceSkeleton(TArray<TWeakObjectPtr<UAnimationAsset>> Objects);
+
 	/** Handle menu item for wanting to open asset in a new editor */
 	void ExecuteOpenInNewWindow(TArray<TWeakObjectPtr<UAnimationAsset>> Objects);
 
