Fix regression that prevents opening a mapA from CB without saving currently unsaved mapB

#rnx

diff --git a/Engine/Source/Developer/AssetTools/Private/AssetTypeActions/AssetTypeActions_World.cpp b/Engine/Source/Developer/AssetTools/Private/AssetTypeActions/AssetTypeActions_World.cpp
--- a/Engine/Source/Developer/AssetTools/Private/AssetTypeActions/AssetTypeActions_World.cpp
+++ b/Engine/Source/Developer/AssetTools/Private/AssetTypeActions/AssetTypeActions_World.cpp
@@ -37,7 +37,7 @@ void FAssetTypeActions_World::OpenAssetEditor( const TArray<UObject*>& InObjects
 		UWorld* World = Cast<UWorld>(*ObjIt);
 		if (World != nullptr && 
 			ensureMsgf(World->GetTypedOuter<UPackage>(), TEXT("World(%s) is not in a package and cannot be opened"), *World->GetFullName()) && 
-			ensureMsgf(!World->GetPackage()->IsDirty(), TEXT("World(%s) is unsaved and cannot be opened")))
+			ensureMsgf(!World->GetPackage()->HasAnyPackageFlags(PKG_NewlyCreated), TEXT("World(%s) is unsaved and cannot be opened")))
 		{
 			const FString FileToOpen = FPackageName::LongPackageNameToFilename(World->GetOutermost()->GetName(), FPackageName::GetMapPackageExtension());
 			const bool bLoadAsTemplate = false;
@@ -75,17 +75,18 @@ TArray<FAssetData> FAssetTypeActions_World::GetValidAssetsForPreviewOrEdit(TArra
 		constexpr bool bPromptUserToSave = true;
 		constexpr bool bSaveMapPackages = true;
 		constexpr bool bSaveContentPackages = true;
-		constexpr bool bFastSave = false;
-		constexpr bool bNotifyNoPackagesSaved = false;
-		constexpr bool bCanBeDeclined = false;
-		if (FEditorFileUtils::SaveDirtyPackages(bPromptUserToSave, bSaveMapPackages, bSaveContentPackages, bFastSave, bNotifyNoPackagesSaved, bCanBeDeclined))
+		if (FEditorFileUtils::SaveDirtyPackages(bPromptUserToSave, bSaveMapPackages, bSaveContentPackages))
 		{
 			// Validate that Asset was saved or isn't loaded meaning it can be loaded
 			const bool bLoad = false;
-			if (UWorld* World = Cast<UWorld>(AssetData.FastGetAsset(bLoad)); !World || !World->GetPackage()->IsDirty())
+			if (UWorld* World = Cast<UWorld>(AssetData.FastGetAsset(bLoad)); !World || !World->GetPackage()->HasAnyPackageFlags(PKG_NewlyCreated))
 			{
 				AssetsToOpen.Add(AssetData);
 			}
+			else
+			{
+				FMessageDialog::Open(EAppMsgType::Ok, LOCTEXT("CannotOpenNewlyCreatedMapWithoutSaving", "The level you are trying to open needs to be saved first."));
+			}
 		}
 	}
 
