Fix Non-unity: missing include

#rnx

diff --git a/Engine/Source/Developer/AssetTools/Private/AssetTypeActions/AssetTypeActions_World.cpp b/Engine/Source/Developer/AssetTools/Private/AssetTypeActions/AssetTypeActions_World.cpp
--- a/Engine/Source/Developer/AssetTools/Private/AssetTypeActions/AssetTypeActions_World.cpp
+++ b/Engine/Source/Developer/AssetTools/Private/AssetTypeActions/AssetTypeActions_World.cpp
@@ -2,6 +2,7 @@
 
 #include "AssetTypeActions/AssetTypeActions_World.h"
 #include "Misc/PackageName.h"
+#include "Misc/MessageDialog.h"
 #include "ThumbnailRendering/WorldThumbnailInfo.h"
 #include "FileHelpers.h"
 
