Rework 13645015:
- Only call FAssetRegistryModule::AssetCreated if necessary
- Call FAssetRegistryModule::AssetCreated before save so that saved packages doesn't end up with PKG_NewlyCreated flag

diff --git a/Engine/Source/Developer/AssetTools/Private/AssetTools.cpp b/Engine/Source/Developer/AssetTools/Private/AssetTools.cpp
--- a/Engine/Source/Developer/AssetTools/Private/AssetTools.cpp
+++ b/Engine/Source/Developer/AssetTools/Private/AssetTools.cpp
@@ -670,8 +670,15 @@ UObject* UAssetToolsImpl::PerformDuplicateAsset(const FString& AssetName, const
 	if(NewObject != nullptr)
 	{
 		// Assets must have RF_Public and RF_Standalone
+		const bool bIsAsset = NewObject->IsAsset();
 		NewObject->SetFlags(RF_Public | RF_Standalone);
 
+		if (!bIsAsset && NewObject->IsAsset())
+		{
+			// Notify the asset registry
+			FAssetRegistryModule::AssetCreated(NewObject);
+		}
+
 		if ( ISourceControlModule::Get().IsEnabled() )
 		{
 			// Save package here if SCC is enabled because the user can use SCC to revert a change
@@ -686,9 +693,6 @@ UObject* UAssetToolsImpl::PerformDuplicateAsset(const FString& AssetName, const
 			SourceControlHelpers::BranchPackage(NewObject->GetOutermost(), OriginalObject->GetOutermost());
 		}
 
-		// Notify the asset registry
-		FAssetRegistryModule::AssetCreated(NewObject);
-
 		// analytics create record
 		UAssetToolsImpl::OnNewCreateRecord(NewObject->GetClass(), true);
 	}
