Fixed crash on interacting with garbage collected packages after their deletion through a sync in source control.

#rnx

diff --git a/Engine/Source/Developer/AssetTools/Private/AssetViewUtils.cpp b/Engine/Source/Developer/AssetTools/Private/AssetViewUtils.cpp
--- a/Engine/Source/Developer/AssetTools/Private/AssetViewUtils.cpp
+++ b/Engine/Source/Developer/AssetTools/Private/AssetViewUtils.cpp
@@ -1789,14 +1789,15 @@ void AssetViewUtils::SyncPathsFromSourceControl(const TArray<FString>& ContentPa
 		}
 
 		// Syncing may have deleted some packages, so we need to unload those rather than re-load them...
-		TArray<UPackage*> PackagesToUnload;
+		// Note: we will store the package using weak pointers here otherwise we might have garbage collection issues after the ReloadPackages call
+		TArray<TWeakObjectPtr<UPackage>> PackagesToUnload;
 		LoadedPackages.RemoveAll([&](UPackage* InPackage) -> bool
 		{
 			const FString PackageExtension = InPackage->ContainsMap() ? FPackageName::GetMapPackageExtension() : FPackageName::GetAssetPackageExtension();
 			const FString PackageFilename = FPackageName::LongPackageNameToFilename(InPackage->GetName(), PackageExtension);
 			if (!FPaths::FileExists(PackageFilename))
 			{
-				PackagesToUnload.Emplace(InPackage);
+				PackagesToUnload.Emplace(MakeWeakObjectPtr(InPackage));
 				return true; // remove package
 			}
 			return false; // keep package
@@ -1812,7 +1813,16 @@ void AssetViewUtils::SyncPathsFromSourceControl(const TArray<FString>& ContentPa
 		UPackageTools::ReloadPackages(LoadedPackages);
 
 		// Unload any deleted packages...
-		UPackageTools::UnloadPackages(PackagesToUnload);
+		TArray<UPackage*> PackageRawPtrsToUnload;
+		for (TWeakObjectPtr<UPackage>& PackageToUnload : PackagesToUnload)
+		{
+			if (PackageToUnload.IsValid())
+			{
+				PackageRawPtrsToUnload.Emplace(PackageToUnload.Get());
+			}
+		}
+
+		UPackageTools::UnloadPackages(PackageRawPtrsToUnload);
 
 		// Re-cache the SCC state...
 		SCCProvider.Execute(ISourceControlOperation::Create<FUpdateStatus>(), PathsOnDisk, EConcurrency::Asynchronous);
