Fixed Instances of a script in the world copied from a template will continue to use the template's version of the class

Originally prepared by markus.breyer with some additional fixes for some issues I was seeing while testing

diff --git a/Engine/Source/Developer/AssetTools/Private/AssetTools.cpp b/Engine/Source/Developer/AssetTools/Private/AssetTools.cpp
--- a/Engine/Source/Developer/AssetTools/Private/AssetTools.cpp
+++ b/Engine/Source/Developer/AssetTools/Private/AssetTools.cpp
@@ -932,13 +932,20 @@ void UAssetToolsImpl::GetAllAdvancedCopySources(FName SelectedPackage, FAdvanced
 	}
 }
 
-bool UAssetToolsImpl::AdvancedCopyPackages(const TMap<FString, FString>& SourceAndDestPackages, const bool bForceAutosave, const bool bCopyOverAllDestinationOverlaps) const
+bool UAssetToolsImpl::AdvancedCopyPackages(
+	const TMap<FString, FString>& SourceAndDestPackages,
+	const bool bForceAutosave,
+	const bool bCopyOverAllDestinationOverlaps,
+	FDuplicatedObjects* OutDuplicatedObjects) const
 {
 	if (ValidateFlattenedAdvancedCopyDestinations(SourceAndDestPackages))
 	{
 		TArray<FString> SuccessfullyCopiedDestinationFiles;
 		TArray<FName> SuccessfullyCopiedSourcePackages;
-		TArray<TMap<TSoftObjectPtr<UObject>, TSoftObjectPtr<UObject>>> DuplicatedObjectsForEachPackage;
+
+		FDuplicatedObjects DuplicatedObjectsLocal;
+		FDuplicatedObjects& DuplicatedObjectsForEachPackage = OutDuplicatedObjects ? *OutDuplicatedObjects : DuplicatedObjectsLocal;
+
 		TSet<UObject*> ExistingObjectSet;
 		TSet<UObject*> NewObjectSet;
 		FString CopyErrors;
@@ -1111,7 +1118,7 @@ bool UAssetToolsImpl::AdvancedCopyPackages(const FAdvancedCopyParams& CopyParams
 	TMap<FString, FString> FlattenedDestinationMap;
 	if (FlattenAdvancedCopyDestinations(PackagesAndDestinations, FlattenedDestinationMap))
 	{
-		return AdvancedCopyPackages(FlattenedDestinationMap, CopyParams.bShouldForceSave, CopyParams.bCopyOverAllDestinationOverlaps);
+		return AdvancedCopyPackages(FlattenedDestinationMap, CopyParams.bShouldForceSave, CopyParams.bCopyOverAllDestinationOverlaps, nullptr);
 	}
 	return false;
 }
diff --git a/Engine/Source/Developer/AssetTools/Private/AssetTools.h b/Engine/Source/Developer/AssetTools/Private/AssetTools.h
--- a/Engine/Source/Developer/AssetTools/Private/AssetTools.h
+++ b/Engine/Source/Developer/AssetTools/Private/AssetTools.h
@@ -112,7 +112,7 @@ public:
 	virtual FAssetPostRenameEvent& OnAssetPostRename() override { return AssetRenameManager->OnAssetPostRenameEvent(); }
 	virtual void ExpandDirectories(const TArray<FString>& Files, const FString& DestinationPath, TArray<TPair<FString, FString>>& FilesAndDestinations) const override;
 	virtual bool AdvancedCopyPackages(const FAdvancedCopyParams& CopyParams, const TArray<TMap<FString, FString>> PackagesAndDestinations) const override;
-	virtual bool AdvancedCopyPackages(const TMap<FString, FString>& SourceAndDestPackages, const bool bForceAutosave, const bool bCopyOverAllDestinationOverlaps) const override;
+	virtual bool AdvancedCopyPackages(const TMap<FString, FString>& SourceAndDestPackages, const bool bForceAutosave, const bool bCopyOverAllDestinationOverlaps, FDuplicatedObjects* OutDuplicatedObjects) const override;
 	virtual void GenerateAdvancedCopyDestinations(FAdvancedCopyParams& InParams, const TArray<FName>& InPackageNamesToCopy, const class UAdvancedCopyCustomization* CopyCustomization, TMap<FString, FString>& OutPackagesAndDestinations) const override;
 	virtual bool FlattenAdvancedCopyDestinations(const TArray<TMap<FString, FString>> PackagesAndDestinations, TMap<FString, FString>& FlattenedPackagesAndDestinations) const override;
 	virtual bool ValidateFlattenedAdvancedCopyDestinations(const TMap<FString, FString>& FlattenedPackagesAndDestinations) const override;
diff --git a/Engine/Source/Developer/AssetTools/Public/IAssetTools.h b/Engine/Source/Developer/AssetTools/Public/IAssetTools.h
--- a/Engine/Source/Developer/AssetTools/Public/IAssetTools.h
+++ b/Engine/Source/Developer/AssetTools/Public/IAssetTools.h
@@ -176,6 +176,9 @@ private:
 	FString DropLocationForAdvancedCopy;
 };
 
+// An array of maps each storing pairs of original object -> duplicated object
+using FDuplicatedObjects = TArray<TMap<TSoftObjectPtr<UObject>, TSoftObjectPtr<UObject>>>;
+
 UINTERFACE(MinimalApi, BlueprintType, meta = (CannotImplementInterfaceInBlueprint))
 class UAssetTools : public UInterface
 {
@@ -433,7 +436,7 @@ public:
 	virtual bool AdvancedCopyPackages(const FAdvancedCopyParams& CopyParams, const TArray<TMap<FString, FString>> PackagesAndDestinations) const = 0;
 
 	/** Copies files after the flattened map of sources and destinations was confirmed */
-	virtual bool AdvancedCopyPackages(const TMap<FString, FString>& SourceAndDestPackages, const bool bForceAutosave = false, const bool bCopyOverAllDestinationOverlaps = true) const = 0;
+	virtual bool AdvancedCopyPackages(const TMap<FString, FString>& SourceAndDestPackages, const bool bForceAutosave = false, const bool bCopyOverAllDestinationOverlaps = true, FDuplicatedObjects* OutDuplicatedObjects = nullptr) const = 0;
 
 	/* Given a set of packages to copy, generate the map of those packages to destination filenames */
 	virtual void GenerateAdvancedCopyDestinations(FAdvancedCopyParams& InParams, const TArray<FName>& InPackageNamesToCopy, const UAdvancedCopyCustomization* CopyCustomization, TMap<FString, FString>& OutPackagesAndDestinations) const = 0;
diff --git a/Engine/Source/Runtime/CoreUObject/Public/UObject/UObjectBase.h b/Engine/Source/Runtime/CoreUObject/Public/UObject/UObjectBase.h
--- a/Engine/Source/Runtime/CoreUObject/Public/UObject/UObjectBase.h
+++ b/Engine/Source/Runtime/CoreUObject/Public/UObject/UObjectBase.h
@@ -230,6 +230,7 @@ private:
 	UObject*						OuterPrivate;
 	
 	friend class FBlueprintCompileReinstancer;
+	friend class FVerseObjectClassReplacer;
 	friend class FContextObjectManager;
 
 #if WITH_EDITOR
