Restricted animation asset's "create asset" submenu in content browser

diff --git a/Engine/Source/Developer/AssetTools/Private/AssetTypeActions/AssetTypeActions_AnimSequence.cpp b/Engine/Source/Developer/AssetTools/Private/AssetTypeActions/AssetTypeActions_AnimSequence.cpp
--- a/Engine/Source/Developer/AssetTools/Private/AssetTypeActions/AssetTypeActions_AnimSequence.cpp
+++ b/Engine/Source/Developer/AssetTools/Private/AssetTypeActions/AssetTypeActions_AnimSequence.cpp
@@ -40,15 +40,23 @@ void FAssetTypeActions_AnimSequence::GetActions(const TArray<UObject*>& InObject
 {
 	auto Sequences = GetTypedWeakObjectPtrs<UAnimSequence>(InObjects);
 
-	// create menu
-	Section.AddSubMenu(
-		"CreateAnimSubmenu",
-		LOCTEXT("CreateAnimSubmenu", "Create"),
-		LOCTEXT("CreateAnimSubmenu_ToolTip", "Create assets from this anim sequence"),
-		FNewMenuDelegate::CreateSP(this, &FAssetTypeActions_AnimSequence::FillCreateMenu, Sequences),
-		false,
-		FSlateIcon(FAppStyle::GetAppStyleSetName(), "Persona.AssetActions.CreateAnimAsset")
-		);
+	IAssetTools& AssetTools = FModuleManager::LoadModuleChecked<FAssetToolsModule>("AssetTools").Get();
+
+	if (AssetTools.IsAssetClassSupported(UAnimComposite::StaticClass()) ||
+		AssetTools.IsAssetClassSupported(UAnimMontage::StaticClass()) ||
+		AssetTools.IsAssetClassSupported(UAnimStreamable::StaticClass()) ||
+		AssetTools.IsAssetClassSupported(UPoseAsset::StaticClass()))
+	{
+		// create menu
+		Section.AddSubMenu(
+			"CreateAnimSubmenu",
+			LOCTEXT("CreateAnimSubmenu", "Create"),
+			LOCTEXT("CreateAnimSubmenu_ToolTip", "Create assets from this anim sequence"),
+			FNewMenuDelegate::CreateSP(this, &FAssetTypeActions_AnimSequence::FillCreateMenu, Sequences),
+			false,
+			FSlateIcon(FAppStyle::GetAppStyleSetName(), "Persona.AssetActions.CreateAnimAsset")
+			);
+	}
 
 	Section.AddMenuEntry(
 		"AnimSequence_ReimportWithNewSource",
@@ -146,48 +154,62 @@ void FAssetTypeActions_AnimSequence::GetActions(const TArray<UObject*>& InObject
 
 void FAssetTypeActions_AnimSequence::FillCreateMenu(FMenuBuilder& MenuBuilder, const TArray<TWeakObjectPtr<UAnimSequence>> Sequences) const
 {
-	MenuBuilder.AddMenuEntry(
-		LOCTEXT("AnimSequence_NewAnimComposite", "Create AnimComposite"),
-		LOCTEXT("AnimSequence_NewAnimCompositeTooltip", "Creates an AnimComposite using the selected anim sequence."),
-		FSlateIcon(FAppStyle::GetAppStyleSetName(), "ClassIcon.AnimComposite"),
-		FUIAction(
-			FExecuteAction::CreateSP(this, &FAssetTypeActions_AnimSequence::ExecuteNewAnimComposite, Sequences),
-			FCanExecuteAction()
-			)
-		);
+	IAssetTools& AssetTools = FModuleManager::LoadModuleChecked<FAssetToolsModule>("AssetTools").Get();
 
-	MenuBuilder.AddMenuEntry(
-		LOCTEXT("AnimSequence_NewAnimMontage", "Create AnimMontage"),
-		LOCTEXT("AnimSequence_NewAnimMontageTooltip", "Creates an AnimMontage using the selected anim sequence."),
-		FSlateIcon(FAppStyle::GetAppStyleSetName(), "ClassIcon.AnimMontage"),
-		FUIAction(
-			FExecuteAction::CreateSP(this, &FAssetTypeActions_AnimSequence::ExecuteNewAnimMontage, Sequences),
-			FCanExecuteAction()
-			)
-		);
+	if(AssetTools.IsAssetClassSupported(UAnimComposite::StaticClass()))
+	{
+		MenuBuilder.AddMenuEntry(
+			LOCTEXT("AnimSequence_NewAnimComposite", "Create AnimComposite"),
+			LOCTEXT("AnimSequence_NewAnimCompositeTooltip", "Creates an AnimComposite using the selected anim sequence."),
+				FSlateIcon(FAppStyle::GetAppStyleSetName(), "ClassIcon.AnimComposite"),
+			FUIAction(
+				FExecuteAction::CreateSP(this, &FAssetTypeActions_AnimSequence::ExecuteNewAnimComposite, Sequences),
+				FCanExecuteAction()
+				)
+			);
+	}
 
-	if(GEnableAnimStreamable == 1)
+	if (AssetTools.IsAssetClassSupported(UAnimMontage::StaticClass()))
 	{
 		MenuBuilder.AddMenuEntry(
-			LOCTEXT("AnimSequence_NewAnimStreamable", "Create AnimStreamable"),
-			LOCTEXT("AnimSequence_NewAnimStreamableTooltip", "Creates an AnimStreamable using the selected anim sequence."),
+			LOCTEXT("AnimSequence_NewAnimMontage", "Create AnimMontage"),
+			LOCTEXT("AnimSequence_NewAnimMontageTooltip", "Creates an AnimMontage using the selected anim sequence."),
 			FSlateIcon(FAppStyle::GetAppStyleSetName(), "ClassIcon.AnimMontage"),
 			FUIAction(
-				FExecuteAction::CreateSP(this, &FAssetTypeActions_AnimSequence::ExecuteNewAnimStreamable, Sequences),
+				FExecuteAction::CreateSP(this, &FAssetTypeActions_AnimSequence::ExecuteNewAnimMontage, Sequences),
 				FCanExecuteAction()
-			)
-		);
+				)
+			);
 	}
 
-	MenuBuilder.AddMenuEntry(
-		LOCTEXT("AnimSequence_NewPoseAsset", "Create PoseAsset"),
-		LOCTEXT("AnimSequence_NewPoseAssetTooltip", "Creates an PoseAsset using the selected anim sequence."),
-		FSlateIcon(FAppStyle::GetAppStyleSetName(), "ClassIcon.PoseAsset"),
-		FUIAction(
-		FExecuteAction::CreateSP(this, &FAssetTypeActions_AnimSequence::ExecuteNewPoseAsset, Sequences),
-		FCanExecuteAction()
-		)
+	if(GEnableAnimStreamable == 1)
+	{
+		if (AssetTools.IsAssetClassSupported(UAnimStreamable::StaticClass()))
+		{
+			MenuBuilder.AddMenuEntry(
+				LOCTEXT("AnimSequence_NewAnimStreamable", "Create AnimStreamable"),
+				LOCTEXT("AnimSequence_NewAnimStreamableTooltip", "Creates an AnimStreamable using the selected anim sequence."),
+				FSlateIcon(FAppStyle::GetAppStyleSetName(), "ClassIcon.AnimMontage"),
+				FUIAction(
+					FExecuteAction::CreateSP(this, &FAssetTypeActions_AnimSequence::ExecuteNewAnimStreamable, Sequences),
+					FCanExecuteAction()
+				)
+			);
+		}
+	}
+
+	if (AssetTools.IsAssetClassSupported(UPoseAsset::StaticClass()))
+	{
+		MenuBuilder.AddMenuEntry(
+			LOCTEXT("AnimSequence_NewPoseAsset", "Create PoseAsset"),
+			LOCTEXT("AnimSequence_NewPoseAssetTooltip", "Creates an PoseAsset using the selected anim sequence."),
+			FSlateIcon(FAppStyle::GetAppStyleSetName(), "ClassIcon.PoseAsset"),
+			FUIAction(
+			FExecuteAction::CreateSP(this, &FAssetTypeActions_AnimSequence::ExecuteNewPoseAsset, Sequences),
+			FCanExecuteAction()
+			)
 		);
+	}
 }
 
 void FAssetTypeActions_AnimSequence::GetResolvedSourceFilePaths(const TArray<UObject*>& TypeAssets, TArray<FString>& OutSourceFilePaths) const
