SourceControl: Group multiple sync operations in one.

diff --git a/Engine/Source/Developer/AssetTools/Private/AssetViewUtils.cpp b/Engine/Source/Developer/AssetTools/Private/AssetViewUtils.cpp
--- a/Engine/Source/Developer/AssetTools/Private/AssetViewUtils.cpp
+++ b/Engine/Source/Developer/AssetTools/Private/AssetViewUtils.cpp
@@ -1781,12 +1781,14 @@ void AssetViewUtils::SyncPathsFromSourceControl(const TArray<FString>& ContentPa
 			}
 		}
 
+		// Group everything...
+		TArray<FString> PathsToSync;
+		PathsToSync.Reserve(PathsOnDisk.Num() + ExtraPackagesToSync.Num());
+		PathsToSync.Append(PathsOnDisk);
+		PathsToSync.Append(SourceControlHelpers::PackageFilenames(ExtraPackagesToSync));
+
 		// Sync everything...
-		SCCProvider.Execute(ISourceControlOperation::Create<FSync>(), PathsOnDisk);
-		if (ExtraPackagesToSync.Num() > 0)
-		{
-			SCCProvider.Execute(ISourceControlOperation::Create<FSync>(), SourceControlHelpers::PackageFilenames(ExtraPackagesToSync));
-		}
+		SCCProvider.Execute(ISourceControlOperation::Create<FSync>(), PathsToSync);
 
 		// Syncing may have deleted some packages, so we need to unload those rather than re-load them...
 		// Note: we will store the package using weak pointers here otherwise we might have garbage collection issues after the ReloadPackages call
