Applied permissions to physics assets and rigs in the skeletal mesh context menu

diff --git a/Engine/Source/Developer/AssetTools/Private/AssetTypeActions/AssetTypeActions_SkeletalMesh.cpp b/Engine/Source/Developer/AssetTools/Private/AssetTypeActions/AssetTypeActions_SkeletalMesh.cpp
--- a/Engine/Source/Developer/AssetTools/Private/AssetTypeActions/AssetTypeActions_SkeletalMesh.cpp
+++ b/Engine/Source/Developer/AssetTools/Private/AssetTypeActions/AssetTypeActions_SkeletalMesh.cpp
@@ -497,6 +497,8 @@ void FAssetTypeActions_SkeletalMesh::GetActions(const TArray<UObject*>& InObject
 
 void FAssetTypeActions_SkeletalMesh::FillCreateMenu(UToolMenu* Menu, TArray<TWeakObjectPtr<USkeletalMesh>> Meshes) const
 {
+	IAssetTools& AssetTools = FModuleManager::LoadModuleChecked<FAssetToolsModule>("AssetTools").Get();
+	if (AssetTools.IsAssetClassSupported(UPhysicsAsset::StaticClass()))
 	{
 		FToolMenuSection& Section = Menu->AddSection("CreatePhysicsAsset", LOCTEXT("CreatePhysicsAssetMenuHeading", "Physics Asset"));
 		Section.AddSubMenu(
