World Partition: added support for advanced copy of partitioned worlds, but only when streaming is disabled.

diff --git a/Engine/Source/Developer/AssetTools/Private/AssetTools.cpp b/Engine/Source/Developer/AssetTools/Private/AssetTools.cpp
--- a/Engine/Source/Developer/AssetTools/Private/AssetTools.cpp
+++ b/Engine/Source/Developer/AssetTools/Private/AssetTools.cpp
@@ -117,6 +117,7 @@
 #include "AssetTypeActions/AssetTypeActions_VectorFieldAnimated.h"
 #include "AssetTypeActions/AssetTypeActions_VectorFieldStatic.h"
 #include "AssetTypeActions/AssetTypeActions_World.h"
+#include "WorldPartition/WorldPartition.h"
 #include "SDiscoveringAssetsDialog.h"
 #include "AssetFixUpRedirectors.h"
 #include "ObjectTools.h"
@@ -986,6 +987,17 @@ bool UAssetToolsImpl::AdvancedCopyPackages(
 						const bool bShouldPromptForDestinationConflict = !bCopyOverAllDestinationOverlaps;
 						TMap<TSoftObjectPtr<UObject>, TSoftObjectPtr<UObject>> DuplicatedObjects;
 
+						// Temp fix for some codepaths that allows advanced copy of world packages. For partitioned worlds, this can only be supported for worlds with
+						// streaming disabled and this code should be removed once the callers switch to the same codepath as editor save as.
+						if (UWorld* World = Cast<UWorld>(ExistingObject))
+						{
+							if (UWorldPartition* WorldPartition = World->GetWorldPartition())
+							{
+								check(!WorldPartition->IsStreamingEnabled());
+								WorldPartition->Initialize(World, FTransform::Identity);
+							}
+						}
+
 						if (UObject* NewObject = ObjectTools::DuplicateSingleObject(ExistingObject, MoveDialogInfo.PGN, ObjectsUserRefusedToFullyLoad, bShouldPromptForDestinationConflict, &DuplicatedObjects))
 						{
 							ExistingObjectSet.Add(ExistingObject);
