Deprecating ANY_PACKAGE.

diff --git a/Engine/Plugins/Runtime/nDisplay/Source/DisplayClusterConfigurator/Private/DisplayClusterConfiguratorVersionUtils.cpp b/Engine/Plugins/Runtime/nDisplay/Source/DisplayClusterConfigurator/Private/DisplayClusterConfiguratorVersionUtils.cpp
--- a/Engine/Plugins/Runtime/nDisplay/Source/DisplayClusterConfigurator/Private/DisplayClusterConfiguratorVersionUtils.cpp
+++ b/Engine/Plugins/Runtime/nDisplay/Source/DisplayClusterConfigurator/Private/DisplayClusterConfiguratorVersionUtils.cpp
@@ -32,8 +32,8 @@ void FDisplayClusterConfiguratorVersionUtils::UpdateBlueprintsToNewVersion()
 	IAssetRegistry& AssetRegistry = FModuleManager::LoadModuleChecked<FAssetRegistryModule>(FName("AssetRegistry")).Get();
 
 	TArray<FAssetData> OutAssets;
-	AssetRegistry.GetAssetsByClass(UDisplayClusterBlueprint::StaticClass()->GetFName(), OutAssets, true);
-	AssetRegistry.GetAssetsByClass(UDisplayClusterConfiguratorEditorData::StaticClass()->GetFName(), OutAssets, true);
+	AssetRegistry.GetAssetsByClass(UDisplayClusterBlueprint::StaticClass()->GetClassPathName(), OutAssets, true);
+	AssetRegistry.GetAssetsByClass(UDisplayClusterConfiguratorEditorData::StaticClass()->GetClassPathName(), OutAssets, true);
 	
 	bool bIsAssetVersionNotSupported = false;
 
diff --git a/Engine/Plugins/Runtime/nDisplay/Source/DisplayClusterConfigurator/Private/Views/NewAsset/SDisplayClusterConfiguratorAssetPickerList.cpp b/Engine/Plugins/Runtime/nDisplay/Source/DisplayClusterConfigurator/Private/Views/NewAsset/SDisplayClusterConfiguratorAssetPickerList.cpp
--- a/Engine/Plugins/Runtime/nDisplay/Source/DisplayClusterConfigurator/Private/Views/NewAsset/SDisplayClusterConfiguratorAssetPickerList.cpp
+++ b/Engine/Plugins/Runtime/nDisplay/Source/DisplayClusterConfigurator/Private/Views/NewAsset/SDisplayClusterConfiguratorAssetPickerList.cpp
@@ -35,7 +35,7 @@ void SDisplayClusterConfiguratorAssetPickerList::Construct(const FArguments& InA
 		AssetPickerConfig.bPreloadAssetsForContextMenu = true;
 		AssetPickerConfig.ThumbnailLabel = EThumbnailLabel::AssetName;
 		
-		AssetPickerConfig.Filter.ClassNames.Add(UDisplayClusterBlueprint::StaticClass()->GetFName());
+		AssetPickerConfig.Filter.ClassPaths.Add(UDisplayClusterBlueprint::StaticClass()->GetClassPathName());
 		AssetPickerConfig.Filter.bRecursiveClasses = true;
 	}
 
@@ -60,7 +60,7 @@ void SDisplayClusterConfiguratorAssetPickerList::OnAssetSelected(const FAssetDat
 
 bool SDisplayClusterConfiguratorAssetPickerList::OnShouldFilterAsset(const FAssetData& InAssetData)
 {
-	if (InAssetData.AssetClass == UDisplayClusterBlueprint::StaticClass()->GetFName())
+	if (InAssetData.AssetClassPath == UDisplayClusterBlueprint::StaticClass()->GetClassPathName())
 	{
 		const FString ParentClassPath = InAssetData.GetTagValueRef<FString>("ParentClass");
 		if (!ParentClassPath.IsEmpty())
diff --git a/Engine/Plugins/Runtime/nDisplay/Source/DisplayClusterConfigurator/Private/Views/NewAsset/SDisplayClusterConfiguratorNewBlueprintDialog.cpp b/Engine/Plugins/Runtime/nDisplay/Source/DisplayClusterConfigurator/Private/Views/NewAsset/SDisplayClusterConfiguratorNewBlueprintDialog.cpp
--- a/Engine/Plugins/Runtime/nDisplay/Source/DisplayClusterConfigurator/Private/Views/NewAsset/SDisplayClusterConfiguratorNewBlueprintDialog.cpp
+++ b/Engine/Plugins/Runtime/nDisplay/Source/DisplayClusterConfigurator/Private/Views/NewAsset/SDisplayClusterConfiguratorNewBlueprintDialog.cpp
@@ -47,7 +47,7 @@ TOptional<FAssetData> SDisplayClusterConfiguratorNewBlueprintDialog::GetSelected
 	TArray<FAssetData> SelectedSystemAssets;
 	for (const FAssetData& SelectedAsset : AllSelectedAssets)
 	{
-		if (SelectedAsset.AssetClass == UDisplayClusterBlueprint::StaticClass()->GetFName())
+		if (SelectedAsset.AssetClassPath == UDisplayClusterBlueprint::StaticClass()->GetClassPathName())
 		{
 			SelectedSystemAssets.Add(SelectedAsset);
 		}
diff --git a/Engine/Plugins/Runtime/nDisplay/Source/DisplayClusterMessageInterceptor/Private/DisplayClusterMessageInterceptor.cpp b/Engine/Plugins/Runtime/nDisplay/Source/DisplayClusterMessageInterceptor/Private/DisplayClusterMessageInterceptor.cpp
--- a/Engine/Plugins/Runtime/nDisplay/Source/DisplayClusterMessageInterceptor/Private/DisplayClusterMessageInterceptor.cpp
+++ b/Engine/Plugins/Runtime/nDisplay/Source/DisplayClusterMessageInterceptor/Private/DisplayClusterMessageInterceptor.cpp
@@ -14,7 +14,7 @@ DEFINE_LOG_CATEGORY(LogDisplayClusterInterception);
 
 namespace DisplayClusterMessageInterceptorUtils
 {
-	static const FInterceptedMessageDescriptor MultiUserInterception(TArray<FName>({ TEXT("ConcertSession_CustomEvent") }), TEXT("ConcertMessageId"));
+	static const FInterceptedMessageDescriptor MultiUserInterception(TArray<FTopLevelAssetPath>({ FTopLevelAssetPath(TEXT("/Script/Concert"), TEXT("ConcertSession_CustomEvent"))}), TEXT("ConcertMessageId"));
 }
 
 FDisplayClusterMessageInterceptor::FDisplayClusterMessageInterceptor()
@@ -67,7 +67,7 @@ void FDisplayClusterMessageInterceptor::Start(TSharedPtr<IMessageBus, ESPMode::T
 		for (const FInterceptedMessageDescriptor& Descriptor : InterceptedMessages)
 		{
 			const FString AnnotationString = Descriptor.Annotation.ToString();
-			for (const FName& MessageType : Descriptor.MessageTypes)
+			for (const FTopLevelAssetPath& MessageType : Descriptor.MessageTypes)
 			{
 				InterceptedBus->Intercept(AsShared(), MessageType);
 				UE_LOG(LogDisplayClusterInterception, Display, TEXT("Intercepting message type: '%s' with annotation '%s'"), *MessageType.ToString(), *AnnotationString);
@@ -83,7 +83,7 @@ void FDisplayClusterMessageInterceptor::Stop()
 {
 	if (bIsIntercepting && InterceptedBus)
 	{
-		InterceptedBus->Unintercept(AsShared(), NAME_All);
+		InterceptedBus->Unintercept(AsShared(), IMessageBus::PATHNAME_All);
 		bIsIntercepting = false;
 		Purge();
 		UE_LOG(LogDisplayClusterInterception, Display, TEXT("Stopping interception of bus messages."));
@@ -197,7 +197,7 @@ bool FDisplayClusterMessageInterceptor::InterceptMessage(const TSharedRef<IMessa
 
 	if (Context->GetForwarder() != Address)
 	{
-		const FName MessageType = Context->GetMessageType();
+		const FTopLevelAssetPath MessageType = Context->GetMessageTypePathName();
 		const FInterceptedMessageDescriptor* Descriptor = InterceptedMessages.FindByPredicate([MessageType](const FInterceptedMessageDescriptor& Other) { return Other.MessageTypes.Contains(MessageType); });
 		if (Descriptor)
 		{
diff --git a/Engine/Plugins/Runtime/nDisplay/Source/DisplayClusterMessageInterceptor/Private/DisplayClusterMessageInterceptor.h b/Engine/Plugins/Runtime/nDisplay/Source/DisplayClusterMessageInterceptor/Private/DisplayClusterMessageInterceptor.h
--- a/Engine/Plugins/Runtime/nDisplay/Source/DisplayClusterMessageInterceptor/Private/DisplayClusterMessageInterceptor.h
+++ b/Engine/Plugins/Runtime/nDisplay/Source/DisplayClusterMessageInterceptor/Private/DisplayClusterMessageInterceptor.h
@@ -17,12 +17,12 @@ DECLARE_LOG_CATEGORY_EXTERN(LogDisplayClusterInterception, Log, All);
 
 struct FInterceptedMessageDescriptor
 {
-	FInterceptedMessageDescriptor(TArray<FName>&& InMessageTypes, FName InAnnotation)
+	FInterceptedMessageDescriptor(TArray<FTopLevelAssetPath>&& InMessageTypes, FName InAnnotation)
 		: MessageTypes(MoveTemp(InMessageTypes))
 		, Annotation(InAnnotation)
 	{}
 
-	TArray<FName> MessageTypes;
+	TArray<FTopLevelAssetPath> MessageTypes;
 	FName Annotation;
 };
 
diff --git a/Engine/Plugins/Tests/EditorTests/Source/EditorTests/Private/UnrealEd/BlueprintAutomationTests.cpp b/Engine/Plugins/Tests/EditorTests/Source/EditorTests/Private/UnrealEd/BlueprintAutomationTests.cpp
--- a/Engine/Plugins/Tests/EditorTests/Source/EditorTests/Private/UnrealEd/BlueprintAutomationTests.cpp
+++ b/Engine/Plugins/Tests/EditorTests/Source/EditorTests/Private/UnrealEd/BlueprintAutomationTests.cpp
@@ -262,7 +262,7 @@ public:
 	{
 		FAssetRegistryModule& AssetRegistryModule = FModuleManager::LoadModuleChecked<FAssetRegistryModule>(TEXT("AssetRegistry"));
 		TArray<FAssetData> ObjectList;
-		AssetRegistryModule.Get().GetAssetsByClass(Class->GetFName(), ObjectList);
+		AssetRegistryModule.Get().GetAssetsByClass(Class->GetClassPathName(), ObjectList);
 
 		for (auto ObjIter=ObjectList.CreateConstIterator(); ObjIter; ++ObjIter)
 		{
@@ -582,7 +582,7 @@ public:
 		AssetFilter.bRecursivePaths = true;
 		if (ClassType != NULL)
 		{
-			AssetFilter.ClassNames.Add(ClassType->GetFName());
+			AssetFilter.ClassPaths.Add(ClassType->GetClassPathName());
 		}
 		
 		TArray<FString> AssetPaths;
@@ -1225,7 +1225,7 @@ bool FBlueprintReparentTest::RunTest(const FString& BlueprintAssetPath)
 		// additionally gather up any blueprints that we explicitly specify though the config
 		for (FAssetData const& AssetData : Assets)
 		{
-			UClass* AssetClass = FindObject<UClass>(ANY_PACKAGE, *AssetData.AssetClass.ToString());
+			UClass* AssetClass = FindObject<UClass>(AssetData.AssetClassPath);
 			TestParentClasses.Add(AssetClass);
 		}
 
@@ -1265,7 +1265,7 @@ void FBlueprintRenameAndCloneTest::GetTests(TArray<FString>& OutBeautifiedNames,
 	FAssetRegistryModule& AssetRegistryModule = FModuleManager::LoadModuleChecked<FAssetRegistryModule>(TEXT("AssetRegistry"));
 	
 	TArray<FAssetData> ObjectList;
-	AssetRegistryModule.Get().GetAssetsByClass(UBlueprint::StaticClass()->GetFName(), ObjectList);
+	AssetRegistryModule.Get().GetAssetsByClass(UBlueprint::StaticClass()->GetClassPathName(), ObjectList);
 
 	for (FAssetData const& Asset : ObjectList)
 	{
diff --git a/Engine/Plugins/Tests/EditorTests/Source/EditorTests/Private/UnrealEd/FbxAutomationTests.cpp b/Engine/Plugins/Tests/EditorTests/Source/EditorTests/Private/UnrealEd/FbxAutomationTests.cpp
--- a/Engine/Plugins/Tests/EditorTests/Source/EditorTests/Private/UnrealEd/FbxAutomationTests.cpp
+++ b/Engine/Plugins/Tests/EditorTests/Source/EditorTests/Private/UnrealEd/FbxAutomationTests.cpp
@@ -699,7 +699,7 @@ bool FFbxImportAssetsAutomationTest::RunTest(const FString& Parameters)
 				int32 MaterialNumber = 0;
 				for (FAssetData AssetData : CreatedAssets)
 				{
-					if (AssetData.AssetClass == UMaterial::StaticClass()->GetFName() || AssetData.AssetClass == UMaterialInstanceConstant::StaticClass()->GetFName())
+					if (AssetData.AssetClassPath == UMaterial::StaticClass()->GetClassPathName() || AssetData.AssetClassPath == UMaterialInstanceConstant::StaticClass()->GetClassPathName())
 					{
 						MaterialNumber++;
 					}
diff --git a/Engine/Plugins/Tests/EditorTests/Source/EditorTests/Private/UnrealEd/PlayMapInPIETests.cpp b/Engine/Plugins/Tests/EditorTests/Source/EditorTests/Private/UnrealEd/PlayMapInPIETests.cpp
--- a/Engine/Plugins/Tests/EditorTests/Source/EditorTests/Private/UnrealEd/PlayMapInPIETests.cpp
+++ b/Engine/Plugins/Tests/EditorTests/Source/EditorTests/Private/UnrealEd/PlayMapInPIETests.cpp
@@ -103,7 +103,7 @@ public:
 
 			TArray<FAssetData> MapList;
 			FARFilter Filter;
-			Filter.ClassNames.Add(UWorld::StaticClass()->GetFName());
+			Filter.ClassPaths.Add(UWorld::StaticClass()->GetClassPathName());
 			Filter.bRecursiveClasses = true;
 			Filter.bIncludeOnlyOnDiskAssets = true;
 			if (AssetRegistry.GetAssets(Filter, /*out*/ MapList))
diff --git a/Engine/Plugins/Tests/InterchangeTests/Source/InterchangeTests/Private/InterchangeImportTest.cpp b/Engine/Plugins/Tests/InterchangeTests/Source/InterchangeTests/Private/InterchangeImportTest.cpp
--- a/Engine/Plugins/Tests/InterchangeTests/Source/InterchangeTests/Private/InterchangeImportTest.cpp
+++ b/Engine/Plugins/Tests/InterchangeTests/Source/InterchangeTests/Private/InterchangeImportTest.cpp
@@ -38,7 +38,7 @@ void FInterchangeImportTest::GetTests(TArray<FString>& OutBeautifiedNames, TArra
 
 	const FAssetRegistryModule& AssetRegistryModule = FModuleManager::LoadModuleChecked<FAssetRegistryModule>(TEXT("AssetRegistry"));
 	TArray<FAssetData> AllTestPlans;
-	AssetRegistryModule.Get().GetAssetsByClass(UInterchangeImportTestPlan::StaticClass()->GetFName(), AllTestPlans, true);
+	AssetRegistryModule.Get().GetAssetsByClass(UInterchangeImportTestPlan::StaticClass()->GetClassPathName(), AllTestPlans, true);
 
 	// Get a list of all paths containing InterchangeTestPlan assets
 	TSet<FName> Paths;
@@ -80,7 +80,7 @@ bool FInterchangeImportTest::RunTest(const FString& Path)
 		FName PackageName = FName(FPaths::GetBaseFilename(Path, false));
 		FName PackagePath = FName(FPaths::GetPath(Path));
 		FName AssetName = FName(FPaths::GetBaseFilename(Path, true));
-		FName ClassName = UInterchangeImportTestPlan::StaticClass()->GetFName();
+		FTopLevelAssetPath ClassName = UInterchangeImportTestPlan::StaticClass()->GetClassPathName();
 		FInterchangeImportTestData& TP = TestPlans.AddDefaulted_GetRef();
 		TP.AssetData = FAssetData(PackageName, PackagePath, AssetName, ClassName);
 		TP.TestPlan = CastChecked<UInterchangeImportTestPlan>(TP.AssetData.GetAsset());
@@ -90,7 +90,7 @@ bool FInterchangeImportTest::RunTest(const FString& Path)
 		// Run tests in parallel on all TestPlan assets in the given directory
 		FARFilter AssetRegistryFilter;
 		AssetRegistryFilter.PackagePaths.Add(FName(Path));
-		AssetRegistryFilter.ClassNames.Add(UInterchangeImportTestPlan::StaticClass()->GetFName());
+		AssetRegistryFilter.ClassPaths.Add(UInterchangeImportTestPlan::StaticClass()->GetClassPathName());
 		AssetRegistryFilter.bRecursivePaths = false;
 
 		const FAssetRegistryModule& AssetRegistryModule = FModuleManager::LoadModuleChecked<FAssetRegistryModule>(TEXT("AssetRegistry"));
diff --git a/Engine/Plugins/VirtualProduction/CameraCalibration/Source/CameraCalibrationEditor/Private/CameraCalibrationEditorModule.cpp b/Engine/Plugins/VirtualProduction/CameraCalibration/Source/CameraCalibrationEditor/Private/CameraCalibrationEditorModule.cpp
--- a/Engine/Plugins/VirtualProduction/CameraCalibration/Source/CameraCalibrationEditor/Private/CameraCalibrationEditorModule.cpp
+++ b/Engine/Plugins/VirtualProduction/CameraCalibration/Source/CameraCalibrationEditor/Private/CameraCalibrationEditorModule.cpp
@@ -118,7 +118,7 @@ void FCameraCalibrationEditorModule::RegisterPlacementModeItems()
 				TEXT("/CameraCalibration/Devices/Tracker/BP_UE_Tracker"),
 				TEXT("/CameraCalibration/Devices/Tracker"),
 				TEXT("BP_UE_Tracker"),
-				TEXT("Blueprint")
+				FTopLevelAssetPath(TEXT("/Script/Engine"), TEXT("Blueprint"))
 			);
 
 			PlaceActors.Add(IPlacementModeModule::Get().RegisterPlaceableItem(Info->UniqueHandle, MakeShared<FPlaceableItem>(
@@ -135,7 +135,7 @@ void FCameraCalibrationEditorModule::RegisterPlacementModeItems()
 				TEXT("/CameraCalibration/Devices/Tracker/BP_UE_Tracker3"),
 				TEXT("/CameraCalibration/Devices/Tracker"),
 				TEXT("BP_UE_Tracker3"),
-				TEXT("Blueprint")
+				FTopLevelAssetPath(TEXT("/Script/Engine"), TEXT("Blueprint"))
 			);
 
 			PlaceActors.Add(IPlacementModeModule::Get().RegisterPlaceableItem(Info->UniqueHandle, MakeShared<FPlaceableItem>(
diff --git a/Engine/Plugins/VirtualProduction/CameraCalibrationCore/Source/CameraCalibrationCoreEditor/Private/CameraCalibrationCoreEditorModule.cpp b/Engine/Plugins/VirtualProduction/CameraCalibrationCore/Source/CameraCalibrationCoreEditor/Private/CameraCalibrationCoreEditorModule.cpp
--- a/Engine/Plugins/VirtualProduction/CameraCalibrationCore/Source/CameraCalibrationCoreEditor/Private/CameraCalibrationCoreEditorModule.cpp
+++ b/Engine/Plugins/VirtualProduction/CameraCalibrationCore/Source/CameraCalibrationCoreEditor/Private/CameraCalibrationCoreEditorModule.cpp
@@ -120,7 +120,7 @@ void FCameraCalibrationCoreEditorModule::RegisterPlacementModeItems()
 				TEXT("/Script/CameraCalibrationCore.CameraCalibrationCheckerboard"),
 				TEXT("/CameraCalibrationCore"),
 				TEXT("CameraCalibrationCheckerboard"),
-				TEXT("Actor")
+				FTopLevelAssetPath(TEXT("/Script/Engine"), TEXT("Actor"))
 			);
 		
 			PlaceActors.Add(IPlacementModeModule::Get().RegisterPlaceableItem(Info->UniqueHandle, MakeShared<FPlaceableItem>(
diff --git a/Engine/Plugins/VirtualProduction/DMX/DMXEngine/Source/DMXEditor/Private/DMXEditorUtils.cpp b/Engine/Plugins/VirtualProduction/DMX/DMXEngine/Source/DMXEditor/Private/DMXEditorUtils.cpp
--- a/Engine/Plugins/VirtualProduction/DMX/DMXEngine/Source/DMXEditor/Private/DMXEditorUtils.cpp
+++ b/Engine/Plugins/VirtualProduction/DMX/DMXEngine/Source/DMXEditor/Private/DMXEditorUtils.cpp
@@ -884,7 +884,7 @@ void FDMXEditorUtils::GetAllAssetsOfClass(UClass* Class, TArray<UObject*>& OutOb
 
 	IAssetRegistry& AssetRegistry = FModuleManager::LoadModuleChecked<FAssetRegistryModule>(FName("AssetRegistry")).Get();
 	TArray<FAssetData> OutAssets;
-	AssetRegistry.GetAssetsByClass(Class->GetFName(), OutAssets, true);
+	AssetRegistry.GetAssetsByClass(Class->GetClassPathName(), OutAssets, true);
 
 	for (const FAssetData& Asset : OutAssets)
 	{
diff --git a/Engine/Plugins/VirtualProduction/DMX/DMXEngine/Source/DMXEditor/Private/Sequencer/DMXLibraryTrackEditor.cpp b/Engine/Plugins/VirtualProduction/DMX/DMXEngine/Source/DMXEditor/Private/Sequencer/DMXLibraryTrackEditor.cpp
--- a/Engine/Plugins/VirtualProduction/DMX/DMXEngine/Source/DMXEditor/Private/Sequencer/DMXLibraryTrackEditor.cpp
+++ b/Engine/Plugins/VirtualProduction/DMX/DMXEngine/Source/DMXEditor/Private/Sequencer/DMXLibraryTrackEditor.cpp
@@ -48,7 +48,7 @@ TSharedRef<SWidget> CreateAssetPicker(FOnAssetSelected OnAssetSelected, FOnAsset
 		AssetPickerConfig.bAllowNullSelection = false;
 		AssetPickerConfig.InitialAssetViewType = EAssetViewType::List;
 		AssetPickerConfig.Filter.bRecursiveClasses = true;
-		AssetPickerConfig.Filter.ClassNames.Add(UDMXLibrary::StaticClass()->GetFName());
+		AssetPickerConfig.Filter.ClassPaths.Add(UDMXLibrary::StaticClass()->GetClassPathName());
 		AssetPickerConfig.SaveSettingsName = TEXT("SequencerAssetPicker");
 		AssetPickerConfig.AdditionalReferencingAssets.Add(FAssetData(Sequence));
 	}
@@ -136,8 +136,8 @@ bool FDMXLibraryTrackEditor::SupportsType(TSubclassOf<UMovieSceneTrack> Type) co
 bool FDMXLibraryTrackEditor::SupportsSequence(UMovieSceneSequence* InSequence) const
 {
 	// DMX Library tracks can be added to level sequences and widget sequences
-	static UClass* LevelSequenceClass = FindObject<UClass>(ANY_PACKAGE, TEXT("LevelSequence"), true);
-	static UClass* WidgetAnimationClass = FindObject<UClass>(ANY_PACKAGE, TEXT("WidgetAnimation"), true);
+	static UClass* LevelSequenceClass = FindObject<UClass>(nullptr, TEXT("/Script/LevelSequence.LevelSequence"), true);
+	static UClass* WidgetAnimationClass = FindObject<UClass>(nullptr, TEXT("/Script/UMG.WidgetAnimation"), true);
 	return InSequence != nullptr &&
 		((LevelSequenceClass != nullptr && InSequence->GetClass()->IsChildOf(LevelSequenceClass)) ||
 		(WidgetAnimationClass != nullptr && InSequence->GetClass()->IsChildOf(WidgetAnimationClass)));
diff --git a/Engine/Plugins/VirtualProduction/DMX/DMXEngine/Source/DMXEditor/Private/Widgets/SAssetPickerButton.cpp b/Engine/Plugins/VirtualProduction/DMX/DMXEngine/Source/DMXEditor/Private/Widgets/SAssetPickerButton.cpp
--- a/Engine/Plugins/VirtualProduction/DMX/DMXEngine/Source/DMXEditor/Private/Widgets/SAssetPickerButton.cpp
+++ b/Engine/Plugins/VirtualProduction/DMX/DMXEngine/Source/DMXEditor/Private/Widgets/SAssetPickerButton.cpp
@@ -159,7 +159,7 @@ TSharedRef<SWidget> SAssetPickerButton::GenerateAssetPicker()
 	FContentBrowserModule& ContentBrowserModule = FModuleManager::Get().LoadModuleChecked<FContentBrowserModule>(TEXT("ContentBrowser"));
 
 	FAssetPickerConfig AssetPickerConfig;
-	AssetPickerConfig.Filter.ClassNames.Add(AllowedClass->GetFName());
+	AssetPickerConfig.Filter.ClassPaths.Add(AllowedClass->GetClassPathName());
 	AssetPickerConfig.bAllowNullSelection = true;
 	AssetPickerConfig.Filter.bRecursiveClasses = true;
 	AssetPickerConfig.OnAssetSelected = FOnAssetSelected::CreateSP(this, &SAssetPickerButton::OnAssetSelectedFromPicker);
@@ -170,13 +170,13 @@ TSharedRef<SWidget> SAssetPickerButton::GenerateAssetPicker()
 	if (AllowedClasses.Num() > 0)
 	{
 		// Clear out the allowed class names and have the pin's metadata override.
-		AssetPickerConfig.Filter.ClassNames.Empty();
-		AssetPickerConfig.Filter.ClassNames = AllowedClasses;
+		AssetPickerConfig.Filter.ClassPaths.Empty();
+		AssetPickerConfig.Filter.ClassPaths = AllowedClasses;
 	}
 
 	if (DisallowedClasses.Num() > 0)
 	{
-		AssetPickerConfig.Filter.RecursiveClassesExclusionSet.Append(DisallowedClasses);
+		AssetPickerConfig.Filter.RecursiveClassPathsExclusionSet.Append(DisallowedClasses);
 	}
 
 	return
diff --git a/Engine/Plugins/VirtualProduction/DMX/DMXEngine/Source/DMXEditor/Public/Widgets/SAssetPickerButton.h b/Engine/Plugins/VirtualProduction/DMX/DMXEngine/Source/DMXEditor/Public/Widgets/SAssetPickerButton.h
--- a/Engine/Plugins/VirtualProduction/DMX/DMXEngine/Source/DMXEditor/Public/Widgets/SAssetPickerButton.h
+++ b/Engine/Plugins/VirtualProduction/DMX/DMXEngine/Source/DMXEditor/Public/Widgets/SAssetPickerButton.h
@@ -37,11 +37,11 @@ public:
 		/** Allows more than a single class on the picker list.
 		 * Overrides AssetClass (doesn't affect the Use Selected on Content Browser button)
 		 */
-		SLATE_ARGUMENT(TArray<FName>, AllowedClasses)
+		SLATE_ARGUMENT(TArray<FTopLevelAssetPath>, AllowedClasses)
 		/** Disallows subclasses from the allowed asset class(es).
 		 * Doesn't affect the Use Selected on Content Browser button.
 		 */
-		SLATE_ARGUMENT(TArray<FName>, DisallowedClasses)
+		SLATE_ARGUMENT(TArray<FTopLevelAssetPath>, DisallowedClasses)
 
 		/** Called to decide the hovered state using the parent's interact able area */
 		SLATE_EVENT(FOnParentIsHovered, OnParentIsHovered)
@@ -56,9 +56,9 @@ public:
 	/** The class to filter the list of assets in the picker menu */
 	TWeakObjectPtr<UClass> AssetClass;
 	/** Allows more than a single class. Overrides AssetClass */
-	TArray<FName> AllowedClasses;
+	TArray<FTopLevelAssetPath> AllowedClasses;
 	/** Disallows subclasses from the allowed asset class(es) */
-	TArray<FName> DisallowedClasses;
+	TArray<FTopLevelAssetPath> DisallowedClasses;
 
 protected:
 	/** Clicked Use button */
diff --git a/Engine/Plugins/VirtualProduction/DMX/DMXEngine/Source/DMXRuntime/Private/DMXSubsystem.cpp b/Engine/Plugins/VirtualProduction/DMX/DMXEngine/Source/DMXRuntime/Private/DMXSubsystem.cpp
--- a/Engine/Plugins/VirtualProduction/DMX/DMXEngine/Source/DMXRuntime/Private/DMXSubsystem.cpp
+++ b/Engine/Plugins/VirtualProduction/DMX/DMXEngine/Source/DMXRuntime/Private/DMXSubsystem.cpp
@@ -795,7 +795,7 @@ void UDMXSubsystem::Initialize(FSubsystemCollectionBase& Collection)
 #if WITH_EDITOR
 void UDMXSubsystem::OnAssetRegistryAddedAsset(const FAssetData& Asset)
 {
-	if (Asset.AssetClass == UDMXLibrary::StaticClass()->GetFName())
+	if (Asset.AssetClassPath == UDMXLibrary::StaticClass()->GetClassPathName())
 	{
 		UObject* AssetObject = Asset.GetAsset();
 		if (UDMXLibrary* Library = Cast<UDMXLibrary>(AssetObject))
@@ -810,7 +810,7 @@ void UDMXSubsystem::OnAssetRegistryAddedAsset(const FAssetData& Asset)
 #if WITH_EDITOR
 void UDMXSubsystem::OnAssetRegistryRemovedAsset(const FAssetData& Asset)
 {
-	if (Asset.AssetClass == UDMXLibrary::StaticClass()->GetFName())
+	if (Asset.AssetClassPath == UDMXLibrary::StaticClass()->GetClassPathName())
 	{
 		UObject* AssetObject = Asset.GetAsset();
 		if (UDMXLibrary* Library = Cast<UDMXLibrary>(AssetObject))
diff --git a/Engine/Plugins/VirtualProduction/DMX/DMXEngine/Source/DMXRuntime/Private/MVR/DMXMVRFixtureActorLibrary.cpp b/Engine/Plugins/VirtualProduction/DMX/DMXEngine/Source/DMXRuntime/Private/MVR/DMXMVRFixtureActorLibrary.cpp
--- a/Engine/Plugins/VirtualProduction/DMX/DMXEngine/Source/DMXRuntime/Private/MVR/DMXMVRFixtureActorLibrary.cpp
+++ b/Engine/Plugins/VirtualProduction/DMX/DMXEngine/Source/DMXRuntime/Private/MVR/DMXMVRFixtureActorLibrary.cpp
@@ -84,8 +84,8 @@ private:
 		FAssetRegistryModule& AssetRegistryModule = FModuleManager::LoadModuleChecked<FAssetRegistryModule>(TEXT("AssetRegistry"));
 		TArray<FAssetData> BlueprintList;
 		FARFilter Filter;
-		Filter.ClassNames.Add(UBlueprint::StaticClass()->GetFName());
-		Filter.RecursiveClassesExclusionSet.FindOrAdd(*USkeletalMesh::StaticClass()->GetName());
+		Filter.ClassPaths.Add(UBlueprint::StaticClass()->GetClassPathName());
+		Filter.RecursiveClassPathsExclusionSet.FindOrAdd(USkeletalMesh::StaticClass()->GetClassPathName());
 		Filter.bRecursiveClasses = true;
 
 		AssetRegistryModule.Get().GetAssets(Filter, BlueprintList);
diff --git a/Engine/Plugins/VirtualProduction/LevelSnapshots/Source/LevelSnapshotsEditor/Private/Util/TakeSnapshotUtil.cpp b/Engine/Plugins/VirtualProduction/LevelSnapshots/Source/LevelSnapshotsEditor/Private/Util/TakeSnapshotUtil.cpp
--- a/Engine/Plugins/VirtualProduction/LevelSnapshots/Source/LevelSnapshotsEditor/Private/Util/TakeSnapshotUtil.cpp
+++ b/Engine/Plugins/VirtualProduction/LevelSnapshots/Source/LevelSnapshotsEditor/Private/Util/TakeSnapshotUtil.cpp
@@ -114,7 +114,7 @@ ULevelSnapshot* SnapshotEditor::TakeLevelSnapshotAndSaveToDisk(UWorld* World, co
     {
     	FAssetRegistryModule& AssetRegistryModule = FModuleManager::LoadModuleChecked<FAssetRegistryModule>("AssetRegistry");
     	FAssetData ExistingAsset = AssetRegistryModule.Get().GetAssetByObjectPath(FName(*(BasePackageName + "." + FileName)));
-    	if ( ExistingAsset.IsValid() && ExistingAsset.AssetClass == *ULevelSnapshot::StaticClass()->GetName() )
+    	if ( ExistingAsset.IsValid() && ExistingAsset.AssetClassPath == ULevelSnapshot::StaticClass()->GetClassPathName() )
     	{
     		EAppReturnType::Type ShouldReplace = FMessageDialog::Open( EAppMsgType::YesNo, FText::Format(LOCTEXT("ReplaceAssetMessage", "{0} already exists. Do you want to replace it?"), FText::FromString(FileName)) );
     		bProceedWithSave = (ShouldReplace == EAppReturnType::Yes);
diff --git a/Engine/Plugins/VirtualProduction/LevelSnapshots/Source/LevelSnapshotsEditor/Private/Views/Input/SLevelSnapshotsEditorContextPicker.cpp b/Engine/Plugins/VirtualProduction/LevelSnapshots/Source/LevelSnapshotsEditor/Private/Views/Input/SLevelSnapshotsEditorContextPicker.cpp
--- a/Engine/Plugins/VirtualProduction/LevelSnapshots/Source/LevelSnapshotsEditor/Private/Views/Input/SLevelSnapshotsEditorContextPicker.cpp
+++ b/Engine/Plugins/VirtualProduction/LevelSnapshots/Source/LevelSnapshotsEditor/Private/Views/Input/SLevelSnapshotsEditorContextPicker.cpp
@@ -95,7 +95,7 @@ TSharedRef<SWidget> SLevelSnapshotsEditorContextPicker::BuildWorldPickerMenu()
 	// Get all worlds
 	const FAssetRegistryModule& AssetRegistryModule = FModuleManager::LoadModuleChecked<FAssetRegistryModule>(TEXT("AssetRegistry"));
 	TArray<FAssetData> WorldAssets;
-	AssetRegistryModule.Get().GetAssetsByClass(UWorld::StaticClass()->GetFName(), WorldAssets);
+	AssetRegistryModule.Get().GetAssetsByClass(UWorld::StaticClass()->GetClassPathName(), WorldAssets);
 
 	MenuBuilder.BeginSection("Other Worlds", LOCTEXT("OtherWorldsHeader", "Other Worlds")); 
 	for (const FAssetData& Asset : WorldAssets)
diff --git a/Engine/Plugins/VirtualProduction/LevelSnapshots/Source/LevelSnapshotsEditor/Private/Widgets/Filter/SSaveAndLoadFilters.cpp b/Engine/Plugins/VirtualProduction/LevelSnapshots/Source/LevelSnapshotsEditor/Private/Widgets/Filter/SSaveAndLoadFilters.cpp
--- a/Engine/Plugins/VirtualProduction/LevelSnapshots/Source/LevelSnapshotsEditor/Private/Widgets/Filter/SSaveAndLoadFilters.cpp
+++ b/Engine/Plugins/VirtualProduction/LevelSnapshots/Source/LevelSnapshotsEditor/Private/Widgets/Filter/SSaveAndLoadFilters.cpp
@@ -149,7 +149,7 @@ TSharedRef<SWidget> SSaveAndLoadFilters::GenerateSaveLoadMenu()
 		FContentBrowserModule& ContentBrowserModule = FModuleManager::Get().LoadModuleChecked<FContentBrowserModule>(TEXT("ContentBrowser"));
 
 		FAssetPickerConfig AssetPickerConfig;
-		AssetPickerConfig.Filter.ClassNames.Add(ULevelSnapshotsFilterPreset::StaticClass()->GetFName());
+		AssetPickerConfig.Filter.ClassPaths.Add(ULevelSnapshotsFilterPreset::StaticClass()->GetClassPathName());
 		AssetPickerConfig.bAllowNullSelection = false;
 		AssetPickerConfig.Filter.bRecursiveClasses = true;
 		AssetPickerConfig.OnAssetSelected.BindStatic(OnSelectPreset, EditorData);
diff --git a/Engine/Plugins/VirtualProduction/LevelSnapshots/Source/LevelSnapshotsEditor/Private/Widgets/SLevelSnapshotsEditorBrowser.cpp b/Engine/Plugins/VirtualProduction/LevelSnapshots/Source/LevelSnapshotsEditor/Private/Widgets/SLevelSnapshotsEditorBrowser.cpp
--- a/Engine/Plugins/VirtualProduction/LevelSnapshots/Source/LevelSnapshotsEditor/Private/Widgets/SLevelSnapshotsEditorBrowser.cpp
+++ b/Engine/Plugins/VirtualProduction/LevelSnapshots/Source/LevelSnapshotsEditor/Private/Widgets/SLevelSnapshotsEditorBrowser.cpp
@@ -40,7 +40,7 @@ void SLevelSnapshotsEditorBrowser::Construct(const FArguments& InArgs, ULevelSna
 		FModuleManager::Get().LoadModuleChecked<FContentBrowserModule>(TEXT("ContentBrowser"));
 
 	FARFilter ARFilter;
-	ARFilter.ClassNames.Add(ULevelSnapshot::StaticClass()->GetFName());
+	ARFilter.ClassPaths.Add(ULevelSnapshot::StaticClass()->GetClassPathName());
 
 	FAssetPickerConfig AssetPickerConfig;
 	AssetPickerConfig.CustomColumns = GetCustomColumns();
diff --git a/Engine/Plugins/VirtualProduction/RemoteControl/Source/RemoteControl/Private/RCVirtualPropertyContainer.cpp b/Engine/Plugins/VirtualProduction/RemoteControl/Source/RemoteControl/Private/RCVirtualPropertyContainer.cpp
--- a/Engine/Plugins/VirtualProduction/RemoteControl/Source/RemoteControl/Private/RCVirtualPropertyContainer.cpp
+++ b/Engine/Plugins/VirtualProduction/RemoteControl/Source/RemoteControl/Private/RCVirtualPropertyContainer.cpp
@@ -181,7 +181,7 @@ FName URCVirtualPropertyContainerBase::GenerateUniquePropertyName(const FName& I
 		return FinalName;
 	};
 
-	const UEnum* Enum = FindObjectChecked<UEnum>(ANY_PACKAGE, TEXT("EPropertyBagPropertyType"));
+	const UEnum* Enum = FindObjectChecked<UEnum>(nullptr, TEXT("/Script/StructUtils.EPropertyBagPropertyType"));
 	const FName DefaultName = Enum->GetValueAsName(InValueType);
 	
 	int32 Index = 0;
diff --git a/Engine/Plugins/VirtualProduction/RemoteControl/Source/RemoteControl/Private/RemoteControlBinding.cpp b/Engine/Plugins/VirtualProduction/RemoteControl/Source/RemoteControl/Private/RemoteControlBinding.cpp
--- a/Engine/Plugins/VirtualProduction/RemoteControl/Source/RemoteControl/Private/RemoteControlBinding.cpp
+++ b/Engine/Plugins/VirtualProduction/RemoteControl/Source/RemoteControl/Private/RemoteControlBinding.cpp
@@ -63,7 +63,7 @@ namespace
 	{
 		if (Args.Num())
 		{
-			if (URemoteControlPreset* Preset = FindObject<URemoteControlPreset>(ANY_PACKAGE, *Args[0]))
+			if (URemoteControlPreset* Preset = FindFirstObject<URemoteControlPreset>(*Args[0], EFindFirstObjectOptions::EnsureIfAmbiguous))
 			{
 				TStringBuilder<1000> Output;
 				for (URemoteControlBinding* Binding : Preset->Bindings)
diff --git a/Engine/Plugins/VirtualProduction/RemoteControl/Source/RemoteControl/Private/RemoteControlModule.cpp b/Engine/Plugins/VirtualProduction/RemoteControl/Source/RemoteControl/Private/RemoteControlModule.cpp
--- a/Engine/Plugins/VirtualProduction/RemoteControl/Source/RemoteControl/Private/RemoteControlModule.cpp
+++ b/Engine/Plugins/VirtualProduction/RemoteControl/Source/RemoteControl/Private/RemoteControlModule.cpp
@@ -115,7 +115,7 @@ namespace RemoteControlUtil
 	{
 		FARFilter Filter;
 		Filter.bIncludeOnlyOnDiskAssets = false;
-		Filter.ClassNames = {URemoteControlPreset::StaticClass()->GetFName()};
+		Filter.ClassPaths = {URemoteControlPreset::StaticClass()->GetClassPathName()};
 		Filter.bRecursivePaths = true;
 
 		return Filter;
@@ -151,7 +151,7 @@ namespace RemoteControlUtil
 		IAssetRegistry& AssetRegistry = FModuleManager::Get().LoadModuleChecked<FAssetRegistryModule>(TEXT("AssetRegistry")).Get();
 
 		TArray<FAssetData> Assets;
-		AssetRegistry.GetAssetsByClass(URemoteControlPreset::StaticClass()->GetFName(), Assets);
+		AssetRegistry.GetAssetsByClass(URemoteControlPreset::StaticClass()->GetClassPathName(), Assets);
 
 		FAssetData* FoundAsset = Assets.FindByPredicate([&PresetName](const FAssetData& InAsset)
 		{
@@ -1597,7 +1597,7 @@ void FRemoteControlModule::CachePresets() const
 
 void FRemoteControlModule::OnAssetAdded(const FAssetData& AssetData)
 {
-	if (AssetData.AssetClass != URemoteControlPreset::StaticClass()->GetFName())
+	if (AssetData.AssetClassPath != URemoteControlPreset::StaticClass()->GetClassPathName())
 	{
 		return;
 	}
@@ -1612,7 +1612,7 @@ void FRemoteControlModule::OnAssetAdded(const FAssetData& AssetData)
 
 void FRemoteControlModule::OnAssetRemoved(const FAssetData& AssetData)
 {
-	if (AssetData.AssetClass != URemoteControlPreset::StaticClass()->GetFName())
+	if (AssetData.AssetClassPath != URemoteControlPreset::StaticClass()->GetClassPathName())
 	{
 		return;
 	}
@@ -1638,7 +1638,7 @@ void FRemoteControlModule::OnAssetRemoved(const FAssetData& AssetData)
 
 void FRemoteControlModule::OnAssetRenamed(const FAssetData& AssetData, const FString&)
 {
-	if (AssetData.AssetClass != URemoteControlPreset::StaticClass()->GetFName())
+	if (AssetData.AssetClassPath != URemoteControlPreset::StaticClass()->GetClassPathName())
 	{
 		return;
 	}
diff --git a/Engine/Plugins/VirtualProduction/RemoteControl/Source/WebRemoteControl/Private/RemoteControlModels.h b/Engine/Plugins/VirtualProduction/RemoteControl/Source/WebRemoteControl/Private/RemoteControlModels.h
--- a/Engine/Plugins/VirtualProduction/RemoteControl/Source/WebRemoteControl/Private/RemoteControlModels.h
+++ b/Engine/Plugins/VirtualProduction/RemoteControl/Source/WebRemoteControl/Private/RemoteControlModels.h
@@ -591,7 +591,7 @@ struct FRCAssetDescription
 	FRCAssetDescription() = default;
 	FRCAssetDescription(const FAssetData& InAsset)
 		: Name(InAsset.AssetName)
-		, Class(InAsset.AssetClass)
+		, Class(InAsset.AssetClassPath)
 		, Path(InAsset.ObjectPath)
 	{
 		Metadata = RemoteControlModels::SanitizeAssetMetadata(InAsset.TagsAndValues.CopyMap());
@@ -601,7 +601,7 @@ struct FRCAssetDescription
 	FName Name;
 
 	UPROPERTY()
-	FName Class;
+	FTopLevelAssetPath Class;
 
 	UPROPERTY()
 	FName Path;
@@ -641,14 +641,14 @@ struct FRCAssetFilter
 	{
 		FARFilter Filter;
 		Filter.PackageNames = PackageNames;
-		Filter.ClassNames = ClassNames;
-		Filter.RecursiveClassesExclusionSet = RecursiveClassesExclusionSet;
+		Filter.ClassPaths = ClassNames;
+		Filter.RecursiveClassPathsExclusionSet = RecursiveClassesExclusionSet;
 		Filter.bRecursiveClasses = RecursiveClasses;
 		Filter.PackagePaths = PackagePaths;
 
 		// Default to a recursive search at root if no filter is specified.
 		if (Filter.PackageNames.Num() == 0
-			&& Filter.ClassNames.Num() == 0
+			&& Filter.ClassPaths.Num() == 0
 			&& Filter.PackagePaths.Num() == 0)
 		{
 			Filter.PackagePaths = { FName("/Game") };
@@ -673,15 +673,15 @@ struct FRCAssetFilter
 
 	/** The filter component for class names. Instances of the specified classes, but not subclasses (by default), will be included. Derived classes will be included only if bRecursiveClasses is true. */
 	UPROPERTY()
-	TArray<FName> ClassNames;
+	TArray<FTopLevelAssetPath> ClassNames;
 
 	/** Only if bRecursiveClasses is true, the results will exclude classes (and subclasses) in this list */
 	UPROPERTY()
-	TSet<FName> RecursiveClassesExclusionSet;
+	TSet<FTopLevelAssetPath> RecursiveClassesExclusionSet;
 
 	/** Only if EnableBlueprintNativeClassFiltering is true, resulting asset will be filtered for dependants of classes in this list. */
 	UPROPERTY()
-	TArray<FName> NativeParentClasses;
+	TArray<FTopLevelAssetPath> NativeParentClasses;
 
 	/** If true, subclasses of ClassNames will also be included and RecursiveClassesExclusionSet will be excluded. */
 	UPROPERTY()
diff --git a/Engine/Plugins/VirtualProduction/RemoteControl/Source/WebRemoteControl/Private/WebRemoteControl.cpp b/Engine/Plugins/VirtualProduction/RemoteControl/Source/WebRemoteControl/Private/WebRemoteControl.cpp
--- a/Engine/Plugins/VirtualProduction/RemoteControl/Source/WebRemoteControl/Private/WebRemoteControl.cpp
+++ b/Engine/Plugins/VirtualProduction/RemoteControl/Source/WebRemoteControl/Private/WebRemoteControl.cpp
@@ -213,9 +213,9 @@ namespace WebRemoteControl
 		{
 			TSet<UClass*> NativeClasses;
 			NativeClasses.Reserve(SearchAssetRequest.Filter.NativeParentClasses.Num());
-			for (FName NativeClass : SearchAssetRequest.Filter.NativeParentClasses)
+			for (FTopLevelAssetPath NativeClass : SearchAssetRequest.Filter.NativeParentClasses)
 			{
-				if (UClass* Class = FindObject<UClass>(ANY_PACKAGE, *NativeClass.ToString()))
+				if (UClass* Class = FindObject<UClass>(NativeClass))
 				{
 					NativeClasses.Add(Class);
 				}
diff --git a/Engine/Plugins/VirtualProduction/SequencerPlaylists/Source/SequencerPlaylists/Private/SequencerPlaylistsWidgets.cpp b/Engine/Plugins/VirtualProduction/SequencerPlaylists/Source/SequencerPlaylists/Private/SequencerPlaylistsWidgets.cpp
--- a/Engine/Plugins/VirtualProduction/SequencerPlaylists/Source/SequencerPlaylists/Private/SequencerPlaylistsWidgets.cpp
+++ b/Engine/Plugins/VirtualProduction/SequencerPlaylists/Source/SequencerPlaylists/Private/SequencerPlaylistsWidgets.cpp
@@ -451,7 +451,7 @@ TSharedRef<SWidget> SSequencerPlaylistPanel::BuildOpenPlaylistMenu()
 		AssetPickerConfig.bSortByPathInColumnView = false;
 
 		AssetPickerConfig.AssetShowWarningText = LOCTEXT("OpenPlaylistNoAssetsWarning", "No Playlists Found");
-		AssetPickerConfig.Filter.ClassNames.Add(USequencerPlaylist::StaticClass()->GetFName());
+		AssetPickerConfig.Filter.ClassPaths.Add(USequencerPlaylist::StaticClass()->GetClassPathName());
 		AssetPickerConfig.OnAssetSelected = FOnAssetSelected::CreateSP(this, &SSequencerPlaylistPanel::OnLoadPlaylist);
 	}
 
@@ -478,7 +478,7 @@ static bool OpenSaveDialog(const FString& InDefaultPath, const FString& InNewNam
 	{
 		SaveAssetDialogConfig.DefaultPath = InDefaultPath;
 		SaveAssetDialogConfig.DefaultAssetName = InNewNameSuggestion;
-		SaveAssetDialogConfig.AssetClassNames.Add(USequencerPlaylist::StaticClass()->GetFName());
+		SaveAssetDialogConfig.AssetClassNames.Add(USequencerPlaylist::StaticClass()->GetClassPathName());
 		SaveAssetDialogConfig.ExistingAssetPolicy = ESaveAssetDialogExistingAssetPolicy::AllowButWarn;
 		SaveAssetDialogConfig.DialogTitleOverride = LOCTEXT("SavePlaylistDialogTitle", "Save Sequencer Playlist");
 	}
diff --git a/Engine/Plugins/VirtualProduction/StageMonitoring/Source/StageMonitor/Private/StageMonitorSessionManager.cpp b/Engine/Plugins/VirtualProduction/StageMonitoring/Source/StageMonitor/Private/StageMonitorSessionManager.cpp
--- a/Engine/Plugins/VirtualProduction/StageMonitoring/Source/StageMonitor/Private/StageMonitorSessionManager.cpp
+++ b/Engine/Plugins/VirtualProduction/StageMonitoring/Source/StageMonitor/Private/StageMonitorSessionManager.cpp
@@ -103,7 +103,7 @@ UScriptStruct* FStageMonitorSessionManager::FindStructFromName_AnyThread(FName T
 	}
 
 	const bool bExactClass = true;
-	UScriptStruct* Struct = FindObjectSafe<UScriptStruct>(ANY_PACKAGE, *TypeName.ToString(), bExactClass);
+	UScriptStruct* Struct = FindFirstObjectSafe<UScriptStruct>(*TypeName.ToString(), bExactClass ? EFindFirstObjectOptions::ExactClass : EFindFirstObjectOptions::None);
 	if (Struct)
 	{
 		TypeCache.FindOrAdd(TypeName) = Struct;
diff --git a/Engine/Plugins/VirtualProduction/Switchboard/Source/SwitchboardEditor/Private/SwitchboardEditorModule.cpp b/Engine/Plugins/VirtualProduction/Switchboard/Source/SwitchboardEditor/Private/SwitchboardEditorModule.cpp
--- a/Engine/Plugins/VirtualProduction/Switchboard/Source/SwitchboardEditor/Private/SwitchboardEditorModule.cpp
+++ b/Engine/Plugins/VirtualProduction/Switchboard/Source/SwitchboardEditor/Private/SwitchboardEditorModule.cpp
@@ -91,7 +91,7 @@ UClass* FSwitchboardEditorModule::GetDisplayClusterRootActorClass()
 	}
 
 	// Try to find it by name.
-	DisplayClusterRootActorClass = FindObject<UClass>(ANY_PACKAGE, TEXT("DisplayClusterRootActor"));
+	DisplayClusterRootActorClass = FindObject<UClass>(nullptr, TEXT("/Script/DisplayCluster.DisplayClusterRootActor"));
 
 	return DisplayClusterRootActorClass.Get();
 }
diff --git a/Engine/Plugins/VirtualProduction/Takes/Source/TakeRecorder/Private/SerializedRecorder.cpp b/Engine/Plugins/VirtualProduction/Takes/Source/TakeRecorder/Private/SerializedRecorder.cpp
--- a/Engine/Plugins/VirtualProduction/Takes/Source/TakeRecorder/Private/SerializedRecorder.cpp
+++ b/Engine/Plugins/VirtualProduction/Takes/Source/TakeRecorder/Private/SerializedRecorder.cpp
@@ -392,7 +392,7 @@
 		}
 		if (Actor)
 		{
-			UClass* InitPossessedObjectClass = FindObject<UClass>(ANY_PACKAGE, *ActorHeader.ClassName);
+			UClass* InitPossessedObjectClass = UClass::TryFindTypeSlow<UClass>(ActorHeader.ClassName);
 			FMovieScenePossessable Possessable(ActorHeader.Label, InitPossessedObjectClass);
 
 			Possessable.SetGuid(ActorHeader.Guid);
@@ -409,7 +409,7 @@
 		//by MakeSpawnableTemplateFromInstance.  
 		//CachedObjectTemplate = CastChecked<AActor>(InMovieSceneSequence->MakeSpawnableTemplateFromInstance(*ActorToRecord, Header.TemplateName));
 		//
-		UClass* SpawnableClass = FindObject<UClass>(ANY_PACKAGE, *ActorHeader.ClassName);
+		UClass* SpawnableClass = UClass::TryFindTypeSlow<UClass>(ActorHeader.ClassName);
 		UObject* NewInstance = NewObject<UObject>(MovieScene, SpawnableClass, FName(*ActorHeader.TemplateName));
 		//this is where UEngine::CopyPropertiesForUnrelatedObjects happens
 		Actor = CastChecked<AActor>(NewInstance);
@@ -515,7 +515,7 @@
 		// first create a possessable for this component to be controlled by
 		UMovieScene* InMovieScene = InMovieSceneSequence->GetMovieScene();
 
-		UClass* InitPossessedObjectClass = FindObject<UClass>(ANY_PACKAGE, *ActorProperty.ClassName);
+		UClass* InitPossessedObjectClass = UClass::TryFindTypeSlow<UClass>(ActorProperty.ClassName);
 		FMovieScenePossessable ChildPossessable(ActorProperty.UObjectName, InitPossessedObjectClass);
 
 		ChildPossessable.SetGuid(ActorProperty.Guid);
diff --git a/Engine/Plugins/VirtualProduction/Takes/Source/TakeRecorder/Private/Widgets/SLevelSequenceTakeEditor.cpp b/Engine/Plugins/VirtualProduction/Takes/Source/TakeRecorder/Private/Widgets/SLevelSequenceTakeEditor.cpp
--- a/Engine/Plugins/VirtualProduction/Takes/Source/TakeRecorder/Private/Widgets/SLevelSequenceTakeEditor.cpp
+++ b/Engine/Plugins/VirtualProduction/Takes/Source/TakeRecorder/Private/Widgets/SLevelSequenceTakeEditor.cpp
@@ -76,7 +76,7 @@ TArray<UClass*> FindRecordingSourceClasses()
 	TArray<FAssetData> ClassList;
 
 	FARFilter Filter;
-	Filter.ClassNames.Add(UTakeRecorderSource::StaticClass()->GetFName());
+	Filter.ClassPaths.Add(UTakeRecorderSource::StaticClass()->GetClassPathName());
 
 	// Include any Blueprint based objects as well, this includes things like Blutilities, UMG, and GameplayAbility objects
 	Filter.bRecursiveClasses = true;
diff --git a/Engine/Plugins/VirtualProduction/Takes/Source/TakeRecorder/Private/Widgets/STakeRecorderPanel.cpp b/Engine/Plugins/VirtualProduction/Takes/Source/TakeRecorder/Private/Widgets/STakeRecorderPanel.cpp
--- a/Engine/Plugins/VirtualProduction/Takes/Source/TakeRecorder/Private/Widgets/STakeRecorderPanel.cpp
+++ b/Engine/Plugins/VirtualProduction/Takes/Source/TakeRecorder/Private/Widgets/STakeRecorderPanel.cpp
@@ -662,7 +662,7 @@ TSharedRef<SWidget> STakeRecorderPanel::OnGeneratePresetsMenu()
 		AssetPickerConfig.bSortByPathInColumnView = false;
 
 		AssetPickerConfig.AssetShowWarningText = LOCTEXT("NoPresets_Warning", "No Presets Found");
-		AssetPickerConfig.Filter.ClassNames.Add(UTakePreset::StaticClass()->GetFName());
+		AssetPickerConfig.Filter.ClassPaths.Add(UTakePreset::StaticClass()->GetClassPathName());
 		AssetPickerConfig.OnAssetSelected = FOnAssetSelected::CreateSP(this, &STakeRecorderPanel::OnImportPreset);
 	}
 
@@ -714,7 +714,7 @@ static bool OpenSaveDialog(const FString& InDefaultPath, const FString& InNewNam
 	{
 		SaveAssetDialogConfig.DefaultPath = InDefaultPath;
 		SaveAssetDialogConfig.DefaultAssetName = InNewNameSuggestion;
-		SaveAssetDialogConfig.AssetClassNames.Add(UTakePreset::StaticClass()->GetFName());
+		SaveAssetDialogConfig.AssetClassNames.Add(UTakePreset::StaticClass()->GetClassPathName());
 		SaveAssetDialogConfig.ExistingAssetPolicy = ESaveAssetDialogExistingAssetPolicy::AllowButWarn;
 		SaveAssetDialogConfig.DialogTitleOverride = LOCTEXT("SaveTakePresetDialogTitle", "Save Take Preset");
 	}
@@ -902,7 +902,7 @@ TSharedRef<SWidget> STakeRecorderPanel::OnOpenSequenceToRecordIntoMenu()
 		AssetPickerConfig.SaveSettingsName = TEXT("TakeRecorderOpenSequenceToRecordInto");
 
 		AssetPickerConfig.AssetShowWarningText = LOCTEXT("NoSequences_Warning", "No Level Sequences Found");
-		AssetPickerConfig.Filter.ClassNames.Add(ULevelSequence::StaticClass()->GetFName());
+		AssetPickerConfig.Filter.ClassPaths.Add(ULevelSequence::StaticClass()->GetClassPathName());
 		AssetPickerConfig.OnAssetSelected = FOnAssetSelected::CreateSP(this, &STakeRecorderPanel::OnOpenSequenceToRecordInto);
 	}
 
diff --git a/Engine/Plugins/VirtualProduction/Takes/Source/TakeRecorderSources/Private/TakeRecorderActorSource.cpp b/Engine/Plugins/VirtualProduction/Takes/Source/TakeRecorderSources/Private/TakeRecorderActorSource.cpp
--- a/Engine/Plugins/VirtualProduction/Takes/Source/TakeRecorderSources/Private/TakeRecorderActorSource.cpp
+++ b/Engine/Plugins/VirtualProduction/Takes/Source/TakeRecorderSources/Private/TakeRecorderActorSource.cpp
@@ -226,7 +226,7 @@
 	
 	// Initialize the header for this actor in the Manifest Serializer for streaming data capture.
 	FName SerializedType("Actor");
-	FActorFileHeader Header(ObjectBindingName, ActorToRecord->GetActorLabel(), SerializedType, ActorToRecord->GetClass()->GetName(), false);
+	FActorFileHeader Header(ObjectBindingName, ActorToRecord->GetActorLabel(), SerializedType, ActorToRecord->GetClass()->GetPathName(), false);
 
 	if (GetRecordToPossessable())
 	{
diff --git a/Engine/Plugins/VirtualProduction/Takes/Source/TakeRecorderSources/Private/TakeRecorderSourcesModule.cpp b/Engine/Plugins/VirtualProduction/Takes/Source/TakeRecorderSources/Private/TakeRecorderSourcesModule.cpp
--- a/Engine/Plugins/VirtualProduction/Takes/Source/TakeRecorderSources/Private/TakeRecorderSourcesModule.cpp
+++ b/Engine/Plugins/VirtualProduction/Takes/Source/TakeRecorderSources/Private/TakeRecorderSourcesModule.cpp
@@ -581,7 +581,7 @@ public:
 					}
 					else
 					{
-						UClass* FoundClass = FindObject<UClass>(ANY_PACKAGE, *Split);
+						UClass* FoundClass = UClass::TryFindTypeSlow<UClass>(Split);
 						if (FoundClass != nullptr)
 						{
 							FindActorsOfClass(FoundClass, InWorld, ActorsToRecord);
diff --git a/Engine/Plugins/VirtualProduction/Takes/Source/TakesCore/Private/TakesCoreBlueprintLibrary.cpp b/Engine/Plugins/VirtualProduction/Takes/Source/TakesCore/Private/TakesCoreBlueprintLibrary.cpp
--- a/Engine/Plugins/VirtualProduction/Takes/Source/TakesCore/Private/TakesCoreBlueprintLibrary.cpp
+++ b/Engine/Plugins/VirtualProduction/Takes/Source/TakesCore/Private/TakesCoreBlueprintLibrary.cpp
@@ -34,7 +34,7 @@ TArray<FAssetData> UTakesCoreBlueprintLibrary::FindTakes(const FString& Slate, i
 	IAssetRegistry& AssetRegistry = FModuleManager::LoadModuleChecked<FAssetRegistryModule>(AssetRegistryConstants::ModuleName).Get();
 
 	FARFilter Filter;
-	Filter.ClassNames.Add(ULevelSequence::StaticClass()->GetFName());
+	Filter.ClassPaths.Add(ULevelSequence::StaticClass()->GetClassPathName());
 	Filter.TagsAndValues.Add(UTakeMetaData::AssetRegistryTag_Slate, Slate);
 
 	TArray<FAssetData> AllAssets;
diff --git a/Engine/Plugins/Web/HttpBlueprint/Source/HttpBlueprintGraph/Private/K2Node_HttpRequest.cpp b/Engine/Plugins/Web/HttpBlueprint/Source/HttpBlueprintGraph/Private/K2Node_HttpRequest.cpp
--- a/Engine/Plugins/Web/HttpBlueprint/Source/HttpBlueprintGraph/Private/K2Node_HttpRequest.cpp
+++ b/Engine/Plugins/Web/HttpBlueprint/Source/HttpBlueprintGraph/Private/K2Node_HttpRequest.cpp
@@ -30,7 +30,7 @@ namespace HttpRequestLiterals
 
 UK2Node_HttpRequest::UK2Node_HttpRequest(const FObjectInitializer& ObjectInitializer)
 	: Super(ObjectInitializer)
-	, VerbEnum(FindObject<UEnum>(ANY_PACKAGE, TEXT("EHttpVerbs")))
+	, VerbEnum(FindObject<UEnum>(nullptr, TEXT("/Script/HttpBlueprint.EHttpVerbs")))
 {
 }
 
@@ -367,7 +367,7 @@ void UK2Node_HttpRequest::HandleBodyInputPin()
 
 UEnum* UK2Node_HttpRequest::GetVerbEnum() const
 {
-	return VerbEnum ? VerbEnum.Get() : FindObjectChecked<UEnum>(ANY_PACKAGE, TEXT("EHttpVerbs"));
+	return VerbEnum ? VerbEnum.Get() : FindObjectChecked<UEnum>(nullptr, TEXT("/Script/HttpBlueprint.EHttpVerbs"));
 }
 
 void UK2Node_HttpRequest::SyncBodyPinType(UEdGraphPin* const Pin) const
diff --git a/Engine/Plugins/Web/HttpBlueprint/Source/HttpBlueprintGraph/Private/K2Node_MakeRequestHeader.cpp b/Engine/Plugins/Web/HttpBlueprint/Source/HttpBlueprintGraph/Private/K2Node_MakeRequestHeader.cpp
--- a/Engine/Plugins/Web/HttpBlueprint/Source/HttpBlueprintGraph/Private/K2Node_MakeRequestHeader.cpp
+++ b/Engine/Plugins/Web/HttpBlueprint/Source/HttpBlueprintGraph/Private/K2Node_MakeRequestHeader.cpp
@@ -104,7 +104,7 @@ UK2Node_MakeRequestHeader::UK2Node_MakeRequestHeader(const FObjectInitializer& O
 	: Super(ObjectInitializer)
 {
 	NumInputs = 1;
-	PresetEnum = FindObjectChecked<UEnum>(ANY_PACKAGE, TEXT("ERequestPresets"));
+	PresetEnum = FindObjectChecked<UEnum>(nullptr, TEXT("/Script/HttpBlueprint.ERequestPresets"));
 
 	UE::HttpBlueprint::Private::MakeRequestHeaderGlobals::Presets =
 	{
diff --git a/Engine/Source/Developer/AssetTools/Private/AdvancedCopyCustomization.cpp b/Engine/Source/Developer/AssetTools/Private/AdvancedCopyCustomization.cpp
--- a/Engine/Source/Developer/AssetTools/Private/AdvancedCopyCustomization.cpp
+++ b/Engine/Source/Developer/AssetTools/Private/AdvancedCopyCustomization.cpp
@@ -27,9 +27,9 @@ UAdvancedCopyCustomization::UAdvancedCopyCustomization(const class FObjectInitia
 
 	FilterForExcludingDependencies.bRecursivePaths = true;
 	FilterForExcludingDependencies.bRecursiveClasses = true;
-	FilterForExcludingDependencies.ClassNames.Add(UWorld::StaticClass()->GetFName());
-	FilterForExcludingDependencies.ClassNames.Add(ULevel::StaticClass()->GetFName());
-	FilterForExcludingDependencies.ClassNames.Add(UMapBuildDataRegistry::StaticClass()->GetFName());
+	FilterForExcludingDependencies.ClassPaths.Add(UWorld::StaticClass()->GetClassPathName());
+	FilterForExcludingDependencies.ClassPaths.Add(ULevel::StaticClass()->GetClassPathName());
+	FilterForExcludingDependencies.ClassPaths.Add(UMapBuildDataRegistry::StaticClass()->GetClassPathName());
 
 }
 
diff --git a/Engine/Source/Developer/AssetTools/Private/AssetTools.cpp b/Engine/Source/Developer/AssetTools/Private/AssetTools.cpp
--- a/Engine/Source/Developer/AssetTools/Private/AssetTools.cpp
+++ b/Engine/Source/Developer/AssetTools/Private/AssetTools.cpp
@@ -182,7 +182,8 @@ UAssetToolsImpl::UAssetToolsImpl(const FObjectInitializer& ObjectInitializer)
 	: Super(ObjectInitializer)
 	, AssetRenameManager(MakeShareable(new FAssetRenameManager))
 	, AssetFixUpRedirectors(MakeShareable(new FAssetFixUpRedirectors))
-	, NextUserCategoryBit(EAssetTypeCategories::FirstUser) 
+	, NextUserCategoryBit(EAssetTypeCategories::FirstUser)
+	, AssetClassPermissionList_DEPRECATED(MakeShared<FNamePermissionList>())
 	, FolderPermissionList(MakeShared<FPathPermissionList>())
 	, WritableFolderPermissionList(MakeShared<FPathPermissionList>())
 	, CreateAssetsAsExternallyReferenceable(true)
@@ -191,7 +192,7 @@ UAssetToolsImpl::UAssetToolsImpl(const FObjectInitializer& ObjectInitializer)
 	GConfig->GetArray(TEXT("AssetTools"), TEXT("SupportedAssetTypes"), SupportedTypesArray, GEditorIni);
 	for (int32 i = 0; i < (int32)EAssetClassAction::AllAssetActions; ++i)
 	{
-		AssetClassPermissionList.Add(MakeShared<FNamePermissionList>());
+		AssetClassPermissionList.Add(MakeShared<FPathPermissionList>());
 		for (const FString& Type : SupportedTypesArray)
 		{
 			AssetClassPermissionList[i]->AddAllowListItem("AssetToolsConfigFile", *Type);
@@ -336,7 +337,7 @@ void UAssetToolsImpl::RegisterAssetTypeActions(const TSharedRef<IAssetTypeAction
 	bool bSupported = false;
 	if (const UClass* SupportedClass = NewActions->GetSupportedClass())
 	{
-		bSupported = GetAssetClassPermissionList(EAssetClassAction::CreateAsset)->PassesFilter(SupportedClass->GetFName());
+		bSupported = GetAssetClassPathPermissionList(EAssetClassAction::CreateAsset)->PassesFilter(SupportedClass->GetFName());
 	}
 	else
 	{
@@ -1424,7 +1425,7 @@ bool UAssetToolsImpl::AssetUsesGenericThumbnail( const FAssetData& AssetData ) c
 		return !RenderInfo || !RenderInfo->Renderer;
 	}
 
-	if ( AssetData.AssetClass == UBlueprint::StaticClass()->GetFName() )
+	if ( AssetData.AssetClassPath == UBlueprint::StaticClass()->GetClassPathName() )
 	{
 		// Unloaded blueprint asset
 		// It would be more correct here to find the rendering info for the generated class,
@@ -1452,7 +1453,7 @@ bool UAssetToolsImpl::AssetUsesGenericThumbnail( const FAssetData& AssetData ) c
 	else
 	{
 		// Unloaded non-blueprint asset. See if the class has a rendering info.
-		UClass* Class = FindObject<UClass>(ANY_PACKAGE, *AssetData.AssetClass.ToString());
+		UClass* Class = FindObject<UClass>(AssetData.AssetClassPath);
 
 		UObject* ClassCDO = nullptr;
 		if (Class != nullptr)
@@ -2297,7 +2298,7 @@ TArray<UObject*> UAssetToolsImpl::ImportAssetsInternal(const TArray<FString>& Fi
 
 							const FString QualifiedName = PackageName + TEXT(".") + Name;
 							FText Reason;
-							if(NumObjectsDeleted == 0 || !IsUniqueObjectName(*QualifiedName, ANY_PACKAGE, Reason))
+							if(NumObjectsDeleted == 0 || !IsGloballyUniqueObjectName(*QualifiedName, &Reason))
 							{
 								// Original object couldn't be deleted
 								const FText Message = FText::Format(LOCTEXT("ImportDeleteFailed", "Failed to delete '{0}'. The asset is referenced by other content."), FText::FromString(PackageName));
@@ -4002,20 +4003,69 @@ TArray<UFactory*> UAssetToolsImpl::GetNewAssetFactories() const
 
 TSharedRef<FNamePermissionList>& UAssetToolsImpl::GetAssetClassPermissionList()
 {
+PRAGMA_DISABLE_DEPRECATION_WARNINGS
 	return GetAssetClassPermissionList(EAssetClassAction::ViewAsset);
+PRAGMA_ENABLE_DEPRECATION_WARNINGS
+	
 }
 
 TSharedRef<FNamePermissionList>& UAssetToolsImpl::GetAssetClassPermissionList(EAssetClassAction AssetClassAction)
 {
+	// Convert path list to name list. Slow and unefficient but keeping it for compatibility (this function is deprecated)
 	if (AssetClassAction < EAssetClassAction::AllAssetActions)
 	{
-		return AssetClassPermissionList[(int32)AssetClassAction];
+		TSharedRef<FPathPermissionList> List = AssetClassPermissionList[(int32)AssetClassAction];
+		TSharedRef<FNamePermissionList> Result = MakeShared<FNamePermissionList>();
+		if (List->IsDenyListAll())
+		{
+			Result->AddDenyListAll(NAME_None);
+		}
+
+		for (const TPair<FString, FPermissionListOwners>& It : List->GetDenyList())
+		{
+			for (const FName OwnerName : It.Value)
+			{
+				Result->AddDenyListItem(OwnerName, FName(*FPackageName::GetShortName(*It.Key)));
+			}
+		}
+
+		for (const TPair<FString, FPermissionListOwners>& It : List->GetAllowList())
+		{
+			FName Name = *FPackageName::GetShortName(*It.Key);
+			if (Result->PassesFilter(Name))
+			{
+				for (const FName OwnerName : It.Value)
+				{
+					Result->AddAllowListItem(OwnerName, Name);
+				}
+			}
+		}
+
+		// Block everything if none of the AllowList paths passed
+		if (Result->GetAllowList().Num() == 0)
+		{
+			Result->AddDenyListAll(NAME_None);
+		}
+
+		AssetClassPermissionList_DEPRECATED = Result;
+		return AssetClassPermissionList_DEPRECATED;
 	}
 	
 	static TSharedRef<FNamePermissionList> Empty = MakeShared<FNamePermissionList>();
 	return Empty;
 }
 
+TSharedRef<FPathPermissionList>& UAssetToolsImpl::GetAssetClassPathPermissionList(EAssetClassAction AssetClassAction)
+{
+	if (AssetClassAction < EAssetClassAction::AllAssetActions)
+	{
+		return AssetClassPermissionList[(int32)AssetClassAction];
+	}
+
+	static TSharedRef<FPathPermissionList> Empty = MakeShared<FPathPermissionList>();
+	return Empty;
+}
+
 void UAssetToolsImpl::AssetClassPermissionListChanged(EAssetClassAction AssetClassAction)
 {
 	switch (AssetClassAction)
@@ -4026,7 +4076,7 @@ void UAssetToolsImpl::AssetClassPermissionListChanged(EAssetClassAction AssetCla
 			bool bSupported = false;
 			if (const UClass* SupportedClass = ActionsIt->GetSupportedClass())
 			{
-				bSupported = GetAssetClassPermissionList(AssetClassAction)->PassesFilter(SupportedClass->GetFName());
+				bSupported = GetAssetClassPathPermissionList(AssetClassAction)->PassesFilter(SupportedClass->GetPathName());
 			}
 			else
 			{
diff --git a/Engine/Source/Developer/AssetTools/Private/AssetTools.h b/Engine/Source/Developer/AssetTools/Private/AssetTools.h
--- a/Engine/Source/Developer/AssetTools/Private/AssetTools.h
+++ b/Engine/Source/Developer/AssetTools/Private/AssetTools.h
@@ -123,8 +123,11 @@ public:
 	virtual void ConvertVirtualTextures(const TArray<UTexture2D*>& Textures, bool bConvertBackToNonVirtual, const TArray<UMaterial*>* RelatedMaterials = nullptr) const override;
 	virtual bool IsAssetClassSupported(const UClass* AssetClass) const override;
 	virtual TArray<UFactory*> GetNewAssetFactories() const override;
+	UE_DEPRECATED(5.1, "Class names are now represented by path names. Please use GetAssetClassPathPermissionList.")
 	virtual TSharedRef<FNamePermissionList>& GetAssetClassPermissionList() override;
-	virtual TSharedRef<FNamePermissionList>& GetAssetClassPermissionList(EAssetClassAction AssetClassAction) override;
+	UE_DEPRECATED(5.1, "Class names are now represented by path names. Please use GetAssetClassPathPermissionList.")
+	TSharedRef<FNamePermissionList>& GetAssetClassPermissionList(EAssetClassAction AssetClassAction);
+	virtual TSharedRef<FPathPermissionList>& GetAssetClassPathPermissionList(EAssetClassAction AssetClassAction) override;
 	virtual TSharedRef<FPathPermissionList>& GetFolderPermissionList() override;
 	virtual TSharedRef<FPathPermissionList>& GetWritableFolderPermissionList() override;
 	virtual bool AllPassWritableFolderFilter(const TArray<FString>& InPaths) const override;
@@ -212,8 +215,11 @@ private:
 	/** The next user category bit to allocate (set to 0 when there are no more bits left) */
 	uint32 NextUserCategoryBit;
 
-	/** Permission lists of assets by class name, one for each EAssetClassAction */
-	TArray<TSharedRef<FNamePermissionList>> AssetClassPermissionList;
+	/** This should be removed with the removal of GetAssetClassPermissionList functions */
+	TSharedRef<FNamePermissionList> AssetClassPermissionList_DEPRECATED;
+
+	/** Permission lists of assets by class path name, one for each EAssetClassAction */
+	TArray<TSharedRef<FPathPermissionList>> AssetClassPermissionList;
 
 	/** Permission list of folder paths */
 	TSharedRef<FPathPermissionList> FolderPermissionList;
diff --git a/Engine/Source/Developer/AssetTools/Private/AssetTypeActions/AssetTypeActions_AnimBlueprint.cpp b/Engine/Source/Developer/AssetTools/Private/AssetTypeActions/AssetTypeActions_AnimBlueprint.cpp
--- a/Engine/Source/Developer/AssetTools/Private/AssetTypeActions/AssetTypeActions_AnimBlueprint.cpp
+++ b/Engine/Source/Developer/AssetTools/Private/AssetTypeActions/AssetTypeActions_AnimBlueprint.cpp
@@ -89,7 +89,7 @@ void FAssetTypeActions_AnimBlueprint::GetActions(const TArray<UObject*>& InObjec
 					AssetPickerConfig.bAllowNullSelection = false;
 					AssetPickerConfig.InitialAssetViewType = EAssetViewType::List;
 					AssetPickerConfig.Filter.bRecursiveClasses = false;
-					AssetPickerConfig.Filter.ClassNames.Add(USkeleton::StaticClass()->GetFName());
+					AssetPickerConfig.Filter.ClassPaths.Add(USkeleton::StaticClass()->GetClassPathName());
 					AssetPickerConfig.OnShouldFilterAsset = FOnShouldFilterAsset::CreateLambda([WeakAnimBlueprint](const FAssetData& AssetData)
 					{
 						if(UAnimBlueprint* LocalAnimBlueprint = WeakAnimBlueprint.Get())
diff --git a/Engine/Source/Developer/AssetTools/Private/AssetTypeActions/AssetTypeActions_Blueprint.cpp b/Engine/Source/Developer/AssetTools/Private/AssetTypeActions/AssetTypeActions_Blueprint.cpp
--- a/Engine/Source/Developer/AssetTools/Private/AssetTypeActions/AssetTypeActions_Blueprint.cpp
+++ b/Engine/Source/Developer/AssetTools/Private/AssetTypeActions/AssetTypeActions_Blueprint.cpp
@@ -297,7 +297,7 @@ TWeakPtr<IClassTypeActions> FAssetTypeActions_Blueprint::GetClassTypeActions(con
 	{
 		UObject* Outer = nullptr;
 		ResolveName(Outer, ParentClassName, false, false);
-		ParentClass = FindObject<UClass>(ANY_PACKAGE, *ParentClassName);
+		ParentClass = FindObject<UClass>(Outer, *ParentClassName);
 	}
 
 	if(ParentClass)
diff --git a/Engine/Source/Developer/AssetTools/Private/AssetTypeActions/AssetTypeActions_BlueprintGeneratedClass.cpp b/Engine/Source/Developer/AssetTools/Private/AssetTypeActions/AssetTypeActions_BlueprintGeneratedClass.cpp
--- a/Engine/Source/Developer/AssetTools/Private/AssetTypeActions/AssetTypeActions_BlueprintGeneratedClass.cpp
+++ b/Engine/Source/Developer/AssetTools/Private/AssetTypeActions/AssetTypeActions_BlueprintGeneratedClass.cpp
@@ -118,9 +118,7 @@ TWeakPtr<IClassTypeActions> FAssetTypeActions_BlueprintGeneratedClass::GetClassT
 	}
 	if (!ParentClassName.IsEmpty())
 	{
-		UObject* Outer = nullptr;
-		ResolveName(Outer, ParentClassName, false, false);
-		ParentClass = FindObject<UClass>(ANY_PACKAGE, *ParentClassName);
+		ParentClass = UClass::TryFindTypeSlow<UClass>(FPackageName::ExportTextPathToObjectPath(ParentClassName));
 	}
 
 	if (ParentClass)
diff --git a/Engine/Source/Developer/AssetTools/Private/AssetTypeActions/AssetTypeActions_Class.cpp b/Engine/Source/Developer/AssetTools/Private/AssetTypeActions/AssetTypeActions_Class.cpp
--- a/Engine/Source/Developer/AssetTools/Private/AssetTypeActions/AssetTypeActions_Class.cpp
+++ b/Engine/Source/Developer/AssetTools/Private/AssetTypeActions/AssetTypeActions_Class.cpp
@@ -174,7 +174,7 @@ void FAssetTypeActions_Class::OpenAssetEditor( const TArray<UObject*>& InObjects
 
 TWeakPtr<IClassTypeActions> FAssetTypeActions_Class::GetClassTypeActions(const FAssetData& AssetData) const
 {
-	UClass* Class = FindObject<UClass>(ANY_PACKAGE, *AssetData.AssetName.ToString());
+	UClass* Class = FindObject<UClass>(nullptr, *AssetData.ObjectPath.ToString());
 	if(Class)
 	{
 		FAssetToolsModule& AssetToolsModule = FAssetToolsModule::GetModule();
diff --git a/Engine/Source/Developer/AssetTools/Private/AssetTypeActions/AssetTypeActions_DataTable.cpp b/Engine/Source/Developer/AssetTools/Private/AssetTypeActions/AssetTypeActions_DataTable.cpp
--- a/Engine/Source/Developer/AssetTools/Private/AssetTypeActions/AssetTypeActions_DataTable.cpp
+++ b/Engine/Source/Developer/AssetTools/Private/AssetTypeActions/AssetTypeActions_DataTable.cpp
@@ -163,8 +163,8 @@ void FAssetTypeActions_DataTable::OpenAssetEditor( const TArray<UObject*>& InObj
 		DataTablesListText.Indent();
 		for (UDataTable* Table : InvalidDataTables)
 		{
-			const FName ResolvedRowStructName = Table->GetRowStructName();
-			DataTablesListText.AppendLineFormat(LOCTEXT("DataTable_MissingRowStructListEntry", "* {0} (Row Structure: {1})"), FText::FromString(Table->GetName()), FText::FromName(ResolvedRowStructName));
+			const FTopLevelAssetPath ResolvedRowStructName = Table->GetRowStructPathName();
+			DataTablesListText.AppendLineFormat(LOCTEXT("DataTable_MissingRowStructListEntry", "* {0} (Row Structure: {1})"), FText::FromString(Table->GetName()), FText::FromString(ResolvedRowStructName.ToString()));
 		}
 
 		FText Title = LOCTEXT("DataTable_MissingRowStructTitle", "Continue?");
@@ -238,7 +238,7 @@ FText FAssetTypeActions_DataTable::GetDisplayNameFromAssetData(const FAssetData&
 		const FAssetDataTagMapSharedView::FFindTagResult RowStructureTag = AssetData.TagsAndValues.FindTag(NAME_RowStructure);
 		if (RowStructureTag.IsSet())
 		{
-			if (UScriptStruct* FoundStruct = FindObject<UScriptStruct>(ANY_PACKAGE, *RowStructureTag.GetValue(), true))
+			if (UScriptStruct* FoundStruct = UClass::TryFindTypeSlow<UScriptStruct>(RowStructureTag.GetValue(), EFindFirstObjectOptions::ExactClass))
 			{
 				return FText::Format(LOCTEXT("DataTableWithRowType", "Data Table ({0})"), FoundStruct->GetDisplayNameText());
 			}
diff --git a/Engine/Source/Developer/AssetTools/Private/AssetViewUtils.cpp b/Engine/Source/Developer/AssetTools/Private/AssetViewUtils.cpp
--- a/Engine/Source/Developer/AssetTools/Private/AssetViewUtils.cpp
+++ b/Engine/Source/Developer/AssetTools/Private/AssetViewUtils.cpp
@@ -1299,7 +1299,7 @@ bool AssetViewUtils::IsValidObjectPathForCreate(const FString& ObjectPath, const
 		// An asset of a different type already exists at this location, inform the user and continue
 		if (ObjectClass && !ExistingAsset.IsInstanceOf(ObjectClass))
 		{
-			OutErrorMessage = FText::Format(LOCTEXT("RenameAssetOtherTypeAlreadyExists", "An asset of type '{0}' already exists at this location with the name '{1}'."), FText::FromName(ExistingAsset.AssetClass), FText::FromString(ObjectName));
+			OutErrorMessage = FText::Format(LOCTEXT("RenameAssetOtherTypeAlreadyExists", "An asset of type '{0}' already exists at this location with the name '{1}'."), FText::FromString(ExistingAsset.AssetClassPath.ToString()), FText::FromString(ObjectName));
 			
 			// Return false to indicate that the user should enter a new name
 			return false;
diff --git a/Engine/Source/Developer/AssetTools/Public/AssetTypeActions_Base.h b/Engine/Source/Developer/AssetTools/Public/AssetTypeActions_Base.h
--- a/Engine/Source/Developer/AssetTools/Public/AssetTypeActions_Base.h
+++ b/Engine/Source/Developer/AssetTools/Public/AssetTypeActions_Base.h
@@ -158,7 +158,9 @@ public:
 	virtual void BuildBackendFilter(FARFilter& InFilter) override
 	{
 		// Add the supported class for this type to a filter
-		InFilter.ClassNames.Add(GetFilterName());
+		FString FilterString(GetFilterName().ToString());
+		checkf(!FPackageName::IsShortPackageName(FilterString), TEXT("Short class names are not supported: %s"), *FilterString);
+		InFilter.ClassPaths.Add(FTopLevelAssetPath(FilterString));
 		InFilter.bRecursiveClasses = true;
 	}
 	
@@ -179,7 +181,7 @@ public:
 
 	virtual FName GetFilterName() const override
 	{
-		return GetSupportedClass()->GetFName();
+		return FName(*GetSupportedClass()->GetPathName());
 	}
 
 protected:
diff --git a/Engine/Source/Developer/AssetTools/Public/IAssetTools.h b/Engine/Source/Developer/AssetTools/Public/IAssetTools.h
--- a/Engine/Source/Developer/AssetTools/Public/IAssetTools.h
+++ b/Engine/Source/Developer/AssetTools/Public/IAssetTools.h
@@ -473,9 +473,9 @@ public:
 	/** Get asset class permission list for the ViewAsset action */
 	UE_DEPRECATED(5.1, "Pass in an EAssetClassAction instead of nothing")
 	virtual TSharedRef<FNamePermissionList>& GetAssetClassPermissionList() = 0;
-	
+
 	/** Get asset class permission list for content browser and other systems */
-	virtual TSharedRef<FNamePermissionList>& GetAssetClassPermissionList(EAssetClassAction AssetClassAction) = 0;
+	virtual TSharedRef<FPathPermissionList>& GetAssetClassPathPermissionList(EAssetClassAction AssetClassAction) = 0;
 
 	/** Get folder permission list for content browser and other systems */
 	virtual TSharedRef<FPathPermissionList>& GetFolderPermissionList() = 0;
diff --git a/Engine/Source/Developer/CookedEditor/Private/CookedEditorPackageManager.cpp b/Engine/Source/Developer/CookedEditor/Private/CookedEditorPackageManager.cpp
--- a/Engine/Source/Developer/CookedEditor/Private/CookedEditorPackageManager.cpp
+++ b/Engine/Source/Developer/CookedEditor/Private/CookedEditorPackageManager.cpp
@@ -98,7 +98,8 @@ FIniCookedEditorPackageManager::FIniCookedEditorPackageManager(bool bInIsCookedC
 	TArray<FString> DisallowedObjectClassNamesToLoad = GetConfigArray(TEXT("DisallowedObjectClassesToLoad"));;
 	for (const FString& ClassName : DisallowedObjectClassNamesToLoad)
 	{
-		UClass* Class = FindObject<UClass>(ANY_PACKAGE, *ClassName);
+		check(FPackageName::IsValidObjectPath(ClassName));
+		UClass* Class = FindObject<UClass>(nullptr, *ClassName);
 		check(Class);
 
 		DisallowedObjectClassesToLoad.Add(Class);
@@ -107,7 +108,8 @@ FIniCookedEditorPackageManager::FIniCookedEditorPackageManager(bool bInIsCookedC
 	TArray<FString> DisallowedAssetClassNamesToGather = GetConfigArray(TEXT("DisallowedAssetClassesToGather"));;
 	for (const FString& ClassName : DisallowedAssetClassNamesToGather)
 	{
-		UClass* Class = FindObject<UClass>(ANY_PACKAGE, *ClassName);
+		check(FPackageName::IsValidObjectPath(ClassName));
+		UClass* Class = FindObject<UClass>(nullptr, *ClassName);
 		check(Class);
 
 		DisallowedAssetClassesToGather.Add(Class);
diff --git a/Engine/Source/Developer/FunctionalTesting/Private/FunctionalTestingModule.cpp b/Engine/Source/Developer/FunctionalTesting/Private/FunctionalTestingModule.cpp
--- a/Engine/Source/Developer/FunctionalTesting/Private/FunctionalTestingModule.cpp
+++ b/Engine/Source/Developer/FunctionalTesting/Private/FunctionalTestingModule.cpp
@@ -114,7 +114,7 @@ void FFunctionalTestingModule::GetMapTests(bool bEditorOnlyTests, TArray<FString
 
 		TArray<FAssetData> MapList;
 		FARFilter Filter;
-		Filter.ClassNames.Add(UWorld::StaticClass()->GetFName());
+		Filter.ClassPaths.Add(UWorld::StaticClass()->GetClassPathName());
 		Filter.bRecursiveClasses = true;
 		Filter.bIncludeOnlyOnDiskAssets = true;
 		if (AssetRegistry.GetAssets(Filter, /*out*/ MapList))
diff --git a/Engine/Source/Developer/Localization/Private/LocalizationConfigurationScript.cpp b/Engine/Source/Developer/Localization/Private/LocalizationConfigurationScript.cpp
--- a/Engine/Source/Developer/Localization/Private/LocalizationConfigurationScript.cpp
+++ b/Engine/Source/Developer/Localization/Private/LocalizationConfigurationScript.cpp
@@ -390,8 +390,7 @@ namespace LocalizationConfigurationScript
 			{
 				if (ExcludeClass.IsValid())
 				{
-					// These are expected to be the short class name (eg, Texture2D) rather than the path name (eg, /Script/Engine.Texture2D)
-					ConfigSection.Add(TEXT("ExcludeClasses"), ExcludeClass.GetAssetName());
+					ConfigSection.Add(TEXT("ExcludeClasses"), ExcludeClass.GetAssetPathString());
 				}
 			}
 			ConfigSection.Add(TEXT("ShouldExcludeDerivedClasses"), Target->Settings.GatherFromPackages.ShouldExcludeDerivedClasses ? TEXT("true") : TEXT("false"));
diff --git a/Engine/Source/Developer/Merge/Private/SMergeAssetPickerView.cpp b/Engine/Source/Developer/Merge/Private/SMergeAssetPickerView.cpp
--- a/Engine/Source/Developer/Merge/Private/SMergeAssetPickerView.cpp
+++ b/Engine/Source/Developer/Merge/Private/SMergeAssetPickerView.cpp
@@ -268,7 +268,7 @@ FReply SMergeAssetPickerPanel::OnBrowseToPickedAsset()
 TSharedRef<SWidget> SMergeAssetPickerPanel::MakeAssetPickerMenu()
 {
 	FAssetPickerConfig AssetPickerConfig;
-	AssetPickerConfig.Filter.ClassNames.Add(AssetClass->GetFName());
+	AssetPickerConfig.Filter.ClassPaths.Add(AssetClass->GetClassPathName());
 	AssetPickerConfig.bAllowNullSelection = true;
 	AssetPickerConfig.Filter.bRecursiveClasses = true;
 	AssetPickerConfig.OnAssetSelected = FOnAssetSelected::CreateSP(this, &SMergeAssetPickerPanel::OnAssetSelected);
diff --git a/Engine/Source/Developer/RigVMDeveloper/Private/RigVMModel/RigVMController.cpp b/Engine/Source/Developer/RigVMDeveloper/Private/RigVMModel/RigVMController.cpp
--- a/Engine/Source/Developer/RigVMDeveloper/Private/RigVMModel/RigVMController.cpp
+++ b/Engine/Source/Developer/RigVMDeveloper/Private/RigVMModel/RigVMController.cpp
@@ -5437,9 +5437,10 @@ URigVMFunctionReferenceNode* URigVMController::PromoteCollapseNodeToFunctionRefe
 
 	// Create Function
 	URigVMLibraryNode* FunctionDefinition = nullptr;
-	if (!InExistingFunctionDefinitionPath.IsEmpty())
+	if (!InExistingFunctionDefinitionPath.IsEmpty() && 
+		ensureAlwaysMsgf(!FPackageName::IsShortPackageName(InExistingFunctionDefinitionPath), TEXT("Expected full path name for function definition path: \"%s\"), *InExistingFunctionDefinitionPath")))
 	{
-		FunctionDefinition = FindObject<URigVMLibraryNode>(ANY_PACKAGE, *InExistingFunctionDefinitionPath);
+		FunctionDefinition = FindObject<URigVMLibraryNode>(nullptr, *InExistingFunctionDefinitionPath);
 	}
 
 	if (FunctionDefinition == nullptr)
@@ -16242,7 +16243,7 @@ static UObject* FindObjectGloballyWithRedirectors(const TCHAR* InObjectName)
 {
 	// Do a global search for the CPP type. Note that searching with ANY_PACKAGE _does not_
 	// apply redirectors. So only if this fails do we apply them manually below.
-	UObject* Object = FindObject<UField>(ANY_PACKAGE, InObjectName);
+	UObject* Object = FindFirstObject<UField>(InObjectName, EFindFirstObjectOptions::EnsureIfAmbiguous);
 	if(Object != nullptr)
 	{
 		return Object;
@@ -16274,7 +16275,7 @@ static UObject* FindObjectGloballyWithRedirectors(const TCHAR* InObjectName)
 	if (Package == nullptr || Object == nullptr)
 	{
 		// Hail Mary pass.
-		Object = FindObject<UField>(ANY_PACKAGE, *RedirectedObjectName);
+		Object = FindFirstObject<UField>(*RedirectedObjectName, EFindFirstObjectOptions::EnsureIfAmbiguous);
 	}
 	return Object;
 }
diff --git a/Engine/Source/Developer/RigVMDeveloper/Private/RigVMModel/RigVMControllerActions.cpp b/Engine/Source/Developer/RigVMDeveloper/Private/RigVMModel/RigVMControllerActions.cpp
--- a/Engine/Source/Developer/RigVMDeveloper/Private/RigVMModel/RigVMControllerActions.cpp
+++ b/Engine/Source/Developer/RigVMDeveloper/Private/RigVMModel/RigVMControllerActions.cpp
@@ -9,7 +9,7 @@
 
 UScriptStruct* FRigVMActionKey::GetScriptStruct() const
 {
-	return FindObjectChecked<UScriptStruct>(ANY_PACKAGE, *ScriptStructPath);		
+	return FindObjectChecked<UScriptStruct>(nullptr, *ScriptStructPath);		
 }
 
 TSharedPtr<FStructOnScope> FRigVMActionKey::GetAction() const
diff --git a/Engine/Source/Developer/RigVMDeveloper/Private/RigVMModel/RigVMPin.cpp b/Engine/Source/Developer/RigVMDeveloper/Private/RigVMModel/RigVMPin.cpp
--- a/Engine/Source/Developer/RigVMDeveloper/Private/RigVMModel/RigVMPin.cpp
+++ b/Engine/Source/Developer/RigVMDeveloper/Private/RigVMModel/RigVMPin.cpp
@@ -1010,7 +1010,7 @@ UObject* URigVMPin::FindObjectFromCPPTypeObjectPath(const FString& InObjectPath)
 
 	// we do this to avoid ambiguous searches for 
 	// common names such as "transform" or "vector"
-	UPackage* Package = ANY_PACKAGE;
+	UPackage* Package = nullptr;
 	FString PackageName;
 	FString CPPTypeObjectName = InObjectPath;
 	if (InObjectPath.Split(TEXT("."), &PackageName, &CPPTypeObjectName))
@@ -1023,7 +1023,7 @@ UObject* URigVMPin::FindObjectFromCPPTypeObjectPath(const FString& InObjectPath)
 		return ObjectWithinPackage;
 	}
 
-	return FindObject<UObject>(ANY_PACKAGE, *InObjectPath);
+	return FindFirstObject<UObject>(*InObjectPath, EFindFirstObjectOptions::NativeFirst | EFindFirstObjectOptions::EnsureIfAmbiguous);
 }
 
 URigVMVariableNode* URigVMPin::GetBoundVariableNode() const
diff --git a/Engine/Source/Developer/VisualGraphUtils/Private/VisualGraphObjectUtils.cpp b/Engine/Source/Developer/VisualGraphUtils/Private/VisualGraphObjectUtils.cpp
--- a/Engine/Source/Developer/VisualGraphUtils/Private/VisualGraphObjectUtils.cpp
+++ b/Engine/Source/Developer/VisualGraphUtils/Private/VisualGraphObjectUtils.cpp
@@ -69,7 +69,7 @@ FAutoConsoleCommandWithWorldAndArgs FCmdVisualGraphUtilsLogInstancesOfClass
 		FString ClipboardContent;
 		for(const FString& ObjectPathName : InParams)
 		{
-			if(UClass* Class = FindObject<UClass>(ANY_PACKAGE, *ObjectPathName, false))
+			if(UClass* Class = UClass::TryFindTypeSlow<UClass>(ObjectPathName))
 			{
 				ForEachObjectOfClass(Class, [&ClipboardContent](UObject* ObjectOfClass)
 				{
@@ -164,7 +164,8 @@ FAutoConsoleCommandWithWorldAndArgs FCmdVisualGraphUtilsCollectReferences
 		TArray<UObject*> OutersToUse;
 		for(const FString& ObjectPathName : ObjectPathNames)
 		{
-			if(UObject* Object = FindObject<UObject>(ANY_PACKAGE, *ObjectPathName, false))
+			UE_CLOG(FPackageName::IsShortPackageName(ObjectPathName), LogVisualGraphUtils, Warning, TEXT("Expected path name but got short name: \"%s\""), *ObjectPathName);
+			if(UObject* Object = FindObject<UObject>(nullptr, *ObjectPathName, false))
 			{
 				Objects.Add(Object);
 
@@ -184,7 +185,8 @@ FAutoConsoleCommandWithWorldAndArgs FCmdVisualGraphUtilsCollectReferences
 		TArray<UObject*> OutersToSkip;
 		for(const FString& SkipPathName : SkipPathNames)
 		{
-			if(UObject* ObjectToSkip = FindObject<UObject>(ANY_PACKAGE, *SkipPathName, false))
+			UE_CLOG(FPackageName::IsShortPackageName(SkipPathName), LogVisualGraphUtils, Warning, TEXT("Expected path name but got short name: \"%s\""), *SkipPathName);
+			if(UObject* ObjectToSkip = FindObject<UObject>(nullptr, *SkipPathName, false))
 			{
 				if(UClass* ClassToSkip = Cast<UClass>(ObjectToSkip))
 				{
@@ -252,7 +254,8 @@ FAutoConsoleCommandWithWorldAndArgs FCmdVisualGraphUtilsCollectTickables
 		TArray<UObject*> OutersToUse;
 		for(const FString& ObjectPathName : ObjectPathNames)
 		{
-			if(UObject* Object = FindObject<UObject>(ANY_PACKAGE, *ObjectPathName, false))
+			UE_CLOG(FPackageName::IsShortPackageName(ObjectPathName), LogVisualGraphUtils, Warning, TEXT("Expected path name but got short name: \"%s\""), *ObjectPathName);
+			if(UObject* Object = FindObject<UObject>(nullptr, *ObjectPathName, false))
 			{
 				Objects.Add(Object);
 			}
diff --git a/Engine/Source/Editor/AIGraph/Private/AIGraphTypes.cpp b/Engine/Source/Editor/AIGraph/Private/AIGraphTypes.cpp
--- a/Engine/Source/Editor/AIGraph/Private/AIGraphTypes.cpp
+++ b/Engine/Source/Editor/AIGraph/Private/AIGraphTypes.cpp
@@ -462,7 +462,7 @@ void FGraphNodeClassHelper::BuildClassGraph()
 	TArray<FAssetData> BlueprintList;
 
 	FARFilter Filter;
-	Filter.ClassNames.Add(UBlueprint::StaticClass()->GetFName());
+	Filter.ClassPaths.Add(UBlueprint::StaticClass()->GetClassPathName());
 	AssetRegistryModule.Get().GetAssets(Filter, BlueprintList);
 
 	for (int32 i = 0; i < BlueprintList.Num(); i++)
@@ -515,16 +515,16 @@ void FGraphNodeClassHelper::UpdateAvailableBlueprintClasses()
 		FAssetRegistryModule& AssetRegistryModule = FModuleManager::GetModuleChecked<FAssetRegistryModule>(TEXT("AssetRegistry"));
 		const bool bSearchSubClasses = true;
 
-		TArray<FName> ClassNames;
-		TSet<FName> DerivedClassNames;
+		TArray<FTopLevelAssetPath> ClassNames;
+		TSet<FTopLevelAssetPath> DerivedClassNames;
 
 		for (TMap<UClass*, int32>::TIterator It(BlueprintClassCount); It; ++It)
 		{
 			ClassNames.Reset();
-			ClassNames.Add(It.Key()->GetFName());
+			ClassNames.Add(It.Key()->GetClassPathName());
 
 			DerivedClassNames.Empty(DerivedClassNames.Num());
-			AssetRegistryModule.Get().GetDerivedClassNames(ClassNames, TSet<FName>(), DerivedClassNames);
+			AssetRegistryModule.Get().GetDerivedClassNames(ClassNames, TSet<FTopLevelAssetPath>(), DerivedClassNames);
 
 			int32& Count = It.Value();
 			Count = DerivedClassNames.Num();
diff --git a/Engine/Source/Editor/AnimGraph/Private/AnimGraphNode_AssetPlayerBase.cpp b/Engine/Source/Editor/AnimGraph/Private/AnimGraphNode_AssetPlayerBase.cpp
--- a/Engine/Source/Editor/AnimGraph/Private/AnimGraphNode_AssetPlayerBase.cpp
+++ b/Engine/Source/Editor/AnimGraph/Private/AnimGraphNode_AssetPlayerBase.cpp
@@ -299,11 +299,11 @@ void UAnimGraphNode_AssetPlayerBase::GetMenuActionsHelper(
 		FARFilter Filter;
 		for(const TSubclassOf<UObject>& AssetType : InAssetTypes)
 		{
-			Filter.ClassNames.Add(AssetType.Get()->GetFName());
+			Filter.ClassPaths.Add(AssetType.Get()->GetClassPathName());
 		}
 		for(const TSubclassOf<UObject>& ExcludedAssetType : InExcludedAssetTypes)
 		{
-			Filter.RecursiveClassesExclusionSet.Add(ExcludedAssetType.Get()->GetFName());
+			Filter.RecursiveClassPathsExclusionSet.Add(ExcludedAssetType.Get()->GetClassPathName());
 		}	
 		Filter.bRecursiveClasses = true;
 		
diff --git a/Engine/Source/Editor/AnimGraph/Private/AnimGraphNode_LinkedAnimLayer.cpp b/Engine/Source/Editor/AnimGraph/Private/AnimGraphNode_LinkedAnimLayer.cpp
--- a/Engine/Source/Editor/AnimGraph/Private/AnimGraphNode_LinkedAnimLayer.cpp
+++ b/Engine/Source/Editor/AnimGraph/Private/AnimGraphNode_LinkedAnimLayer.cpp
@@ -522,7 +522,7 @@ bool UAnimGraphNode_LinkedAnimLayer::OnShouldFilterInstanceBlueprint(const FAsse
 					CurrentAssetData = AssetRegistryModule.Get().GetAssetByObjectPath(BlueprintPath);
 				}
 			// Only continue checking if the parent is an anim blueprint
-			} while (!bMatchesInterface && (CurrentAssetData.AssetClass == UAnimBlueprint::StaticClass()->GetFName()));
+			} while (!bMatchesInterface && (CurrentAssetData.AssetClassPath == UAnimBlueprint::StaticClass()->GetClassPathName()));
 
 			if(!bMatchesInterface)
 			{
diff --git a/Engine/Source/Editor/AnimGraph/Private/AnimGraphNode_Mirror.cpp b/Engine/Source/Editor/AnimGraph/Private/AnimGraphNode_Mirror.cpp
--- a/Engine/Source/Editor/AnimGraph/Private/AnimGraphNode_Mirror.cpp
+++ b/Engine/Source/Editor/AnimGraph/Private/AnimGraphNode_Mirror.cpp
@@ -165,7 +165,7 @@ void UAnimGraphNode_Mirror::GetMenuActions(FBlueprintActionDatabaseRegistrar& Ac
 		FAssetRegistryModule& AssetRegistryModule = FModuleManager::LoadModuleChecked<FAssetRegistryModule>(TEXT("AssetRegistry"));
 		// define a filter to help in pulling UMirrrorDataTable asset data from the registry
 		FARFilter Filter;
-		Filter.ClassNames.Add(UMirrorDataTable::StaticClass()->GetFName());
+		Filter.ClassPaths.Add(UMirrorDataTable::StaticClass()->GetClassPathName());
 
 		Filter.bRecursiveClasses = true;
 		// Find matching assets and add an entry for each one
@@ -209,7 +209,7 @@ void UAnimGraphNode_Mirror::GetMenuActions(FBlueprintActionDatabaseRegistrar& Ac
 	{
 		FAssetRegistryModule& AssetRegistryModule = FModuleManager::LoadModuleChecked<FAssetRegistryModule>(TEXT("AssetRegistry"));
 		FARFilter Filter;
-		Filter.ClassNames.Add(UMirrorDataTable::StaticClass()->GetFName());
+		Filter.ClassPaths.Add(UMirrorDataTable::StaticClass()->GetClassPathName());
 		Filter.bRecursiveClasses = true;
 		// Find matching assets and add an entry for each one
 		TArray<FAssetData> MirrorDataTableList;
@@ -248,7 +248,7 @@ bool UAnimGraphNode_Mirror::HasMirrorDataTableForBlueprints(const TArray<UBluepr
 	// check to see if any mirror data table is available for the blueprint
 	FAssetRegistryModule& AssetRegistryModule = FModuleManager::LoadModuleChecked<FAssetRegistryModule>(TEXT("AssetRegistry"));
 	FARFilter Filter;
-	Filter.ClassNames.Add(UMirrorDataTable::StaticClass()->GetFName());
+	Filter.ClassPaths.Add(UMirrorDataTable::StaticClass()->GetClassPathName());
 
 	Filter.bRecursiveClasses = true;
 	TArray<FAssetData> MirrorDataTableList;
diff --git a/Engine/Source/Editor/BlueprintGraph/Private/BlueprintActionFilter.cpp b/Engine/Source/Editor/BlueprintGraph/Private/BlueprintActionFilter.cpp
--- a/Engine/Source/Editor/BlueprintGraph/Private/BlueprintActionFilter.cpp
+++ b/Engine/Source/Editor/BlueprintGraph/Private/BlueprintActionFilter.cpp
@@ -550,7 +550,7 @@ static const TArray< TSubclassOf<UEdGraphNode> >& BlueprintActionFilterImpl::Get
 		HiddenNodeTypes.Reserve(HiddenClassNames.Num());
 		for (FString const& ClassName : HiddenClassNames)
 		{
-			if (UClass* FoundClass = FindObject<UClass>(ANY_PACKAGE, *ClassName))
+			if (UClass* FoundClass = UClass::TryFindTypeSlow<UClass>(ClassName))
 			{
 				HiddenNodeTypes.Add(FoundClass);
 			}
diff --git a/Engine/Source/Editor/BlueprintGraph/Private/EdGraphSchema_K2.cpp b/Engine/Source/Editor/BlueprintGraph/Private/EdGraphSchema_K2.cpp
--- a/Engine/Source/Editor/BlueprintGraph/Private/EdGraphSchema_K2.cpp
+++ b/Engine/Source/Editor/BlueprintGraph/Private/EdGraphSchema_K2.cpp
@@ -390,7 +390,7 @@ public:
 
 		{
 			TArray<FAssetData> AssetData;
-			AssetRegistryModule.Get().GetAssetsByClass(UUserDefinedEnum::StaticClass()->GetFName(), AssetData);
+			AssetRegistryModule.Get().GetAssetsByClass(UUserDefinedEnum::StaticClass()->GetClassPathName(), AssetData);
 
 			FTypesDatabase::FUnLoadedTypesList UnLoadedTypesList = MakeShareable(new TArray<FUnloadedAssetData>());
 			for (int32 AssetIndex = 0; AssetIndex < AssetData.Num(); ++AssetIndex)
@@ -409,7 +409,7 @@ public:
 		{
 			{
 				TArray<FAssetData> AssetData;
-				AssetRegistryModule.Get().GetAssetsByClass(UUserDefinedStruct::StaticClass()->GetFName(), AssetData);
+				AssetRegistryModule.Get().GetAssetsByClass(UUserDefinedStruct::StaticClass()->GetClassPathName(), AssetData);
 
 				FTypesDatabase::FUnLoadedTypesList UnLoadedTypesList = MakeShareable(new TArray<FUnloadedAssetData>());
 				for (int32 AssetIndex = 0; AssetIndex < AssetData.Num(); ++AssetIndex)
@@ -426,7 +426,7 @@ public:
 
 			{
 				TArray<FAssetData> AssetData;
-				AssetRegistryModule.Get().GetAssetsByClass(UBlueprint::StaticClass()->GetFName(), AssetData);
+				AssetRegistryModule.Get().GetAssetsByClass(UBlueprint::StaticClass()->GetClassPathName(), AssetData);
 
 				const FString BPInterfaceTypeAllowed(TEXT("BPTYPE_Interface"));
 				const FString BPNormalTypeAllowed(TEXT("BPTYPE_Normal"));
@@ -3553,7 +3553,7 @@ bool UEdGraphSchema_K2::ConvertPropertyToPinType(const FProperty* Property, /*ou
 		if(!BitmaskEnumName.IsEmpty())
 		{
 			// @TODO: Potentially replace this with a serialized UEnum reference on the FProperty (e.g. FByteProperty::Enum)
-			TypeOut.PinSubCategoryObject = FindObject<UEnum>(ANY_PACKAGE, *BitmaskEnumName);
+			TypeOut.PinSubCategoryObject = UClass::TryFindTypeSlow<UEnum>(BitmaskEnumName);
 		}
 	}
 
@@ -5061,7 +5061,8 @@ bool UEdGraphSchema_K2::FindFunctionParameterDefaultValue(const UFunction* Funct
 		// If the parameter is a class then try and get the full name as the metadata might just be the short name
 		if (Param->IsA<FClassProperty>() && !FPackageName::IsValidObjectPath(OutString))
 		{
-			if (UClass* DefaultClass = FindObject<UClass>(ANY_PACKAGE, *OutString, true))
+			UE_LOG(LogBlueprint, Warning, TEXT("Short class name \"%s\" in meta data \"%s\" for function %s"), *OutString, *Param->GetName(), *Function->GetPathName());
+			if (UClass* DefaultClass = FindFirstObject<UClass>(*OutString, EFindFirstObjectOptions::None, ELogVerbosity::Warning, TEXT("UEdGraphSchema_K2::FindFunctionParameterDefaultValue")))
 			{
 				OutString = DefaultClass->GetPathName();
 			}
@@ -6157,7 +6158,7 @@ struct FBackwardCompatibilityConversionHelper
 			, ClassPinName(FunctionRedirect.ClassParamName)
 			, FuncScope(NULL)
 		{
-			FuncScope = FindObject<UClass>(ANY_PACKAGE, *FunctionRedirect.ClassName);
+			FuncScope = FindFirstObject<UClass>(*FunctionRedirect.ClassName, EFindFirstObjectOptions::None, ELogVerbosity::Fatal, TEXT("looking for FunctionRedirect.ClassName"));			
 		}
 
 	};
diff --git a/Engine/Source/Editor/BlueprintGraph/Private/K2Node_BaseAsyncTask.cpp b/Engine/Source/Editor/BlueprintGraph/Private/K2Node_BaseAsyncTask.cpp
--- a/Engine/Source/Editor/BlueprintGraph/Private/K2Node_BaseAsyncTask.cpp
+++ b/Engine/Source/Editor/BlueprintGraph/Private/K2Node_BaseAsyncTask.cpp
@@ -653,7 +653,7 @@
 					FParse::Value(*It.Value().GetValue(), TEXT("OldPinName="), OldPinString);
 					FParse::Value(*It.Value().GetValue(), TEXT("NewPinName="), NewPinString);
 
-					UClass* RedirectProxyClass = FindObject<UClass>(ANY_PACKAGE, *ProxyClassString);
+					UClass* RedirectProxyClass = UClass::TryFindTypeSlow<UClass>(ProxyClassString);
 					if (RedirectProxyClass)
 					{
 						FAsyncTaskPinRedirectMapInfo& PinRedirectInfo = AsyncTaskPinRedirectMap.FindOrAdd(*OldPinString);
diff --git a/Engine/Source/Editor/Blutility/Private/BlutilityContentBrowserExtensions.cpp b/Engine/Source/Editor/Blutility/Private/BlutilityContentBrowserExtensions.cpp
--- a/Engine/Source/Editor/Blutility/Private/BlutilityContentBrowserExtensions.cpp
+++ b/Engine/Source/Editor/Blutility/Private/BlutilityContentBrowserExtensions.cpp
@@ -44,7 +44,7 @@ public:
 		{
 			// Check blueprint utils (we need to load them to query their validity against these assets)
 			TArray<FAssetData> UtilAssets;
-			FBlutilityMenuExtensions::GetBlutilityClasses(UtilAssets, UAssetActionUtility::StaticClass()->GetFName());
+			FBlutilityMenuExtensions::GetBlutilityClasses(UtilAssets, UAssetActionUtility::StaticClass()->GetClassPathName());
 
 			// Collect all UAssetActionUtility derived classes
 			TSet<UAssetActionUtility*> AssetClasses;
diff --git a/Engine/Source/Editor/Blutility/Private/BlutilityMenuExtensions.cpp b/Engine/Source/Editor/Blutility/Private/BlutilityMenuExtensions.cpp
--- a/Engine/Source/Editor/Blutility/Private/BlutilityMenuExtensions.cpp
+++ b/Engine/Source/Editor/Blutility/Private/BlutilityMenuExtensions.cpp
@@ -172,21 +172,21 @@ class SFunctionParamDialog : public SCompoundWidget
 	bool bOKPressed;
 };
 
-void FBlutilityMenuExtensions::GetBlutilityClasses(TArray<FAssetData>& OutAssets, const FName& InClassName)
+void FBlutilityMenuExtensions::GetBlutilityClasses(TArray<FAssetData>& OutAssets, FTopLevelAssetPath InClassName)
 {
 	FAssetRegistryModule& AssetRegistryModule = FModuleManager::LoadModuleChecked<FAssetRegistryModule>(TEXT("AssetRegistry"));
 	IAssetRegistry& AssetRegistry = AssetRegistryModule.Get();
 
 	// Get class names
-	TArray<FName> BaseNames;
+	TArray<FTopLevelAssetPath> BaseNames;
 	BaseNames.Add(InClassName);
-	TSet<FName> Excluded;
-	TSet<FName> DerivedNames;
+	TSet<FTopLevelAssetPath> Excluded;
+	TSet<FTopLevelAssetPath> DerivedNames;
 	AssetRegistry.GetDerivedClassNames(BaseNames, Excluded, DerivedNames);
 
 	// Now get all UEditorUtilityBlueprint assets
 	FARFilter Filter;
-	Filter.ClassNames.Add(UEditorUtilityBlueprint::StaticClass()->GetFName());
+	Filter.ClassPaths.Add(UEditorUtilityBlueprint::StaticClass()->GetClassPathName());
 	Filter.bRecursiveClasses = true;
 	Filter.bRecursivePaths = true;
 
@@ -199,10 +199,9 @@ void FBlutilityMenuExtensions::GetBlutilityClasses(TArray<FAssetData>& OutAssets
 		FAssetDataTagMapSharedView::FFindTagResult Result = Asset.TagsAndValues.FindTag(FBlueprintTags::GeneratedClassPath);
 		if (Result.IsSet())
 		{
-			const FString ClassObjectPath = FPackageName::ExportTextPathToObjectPath(Result.GetValue());
-			const FString ClassName = FPackageName::ObjectPathToObjectName(ClassObjectPath);
+			const FTopLevelAssetPath ClassObjectPath(FPackageName::ExportTextPathToObjectPath(Result.GetValue()));
 
-			if (DerivedNames.Contains(*ClassName))
+			if (DerivedNames.Contains(ClassObjectPath))
 			{
 				OutAssets.Add(Asset);
 			}
diff --git a/Engine/Source/Editor/Blutility/Private/BlutilityMenuExtensions.h b/Engine/Source/Editor/Blutility/Private/BlutilityMenuExtensions.h
--- a/Engine/Source/Editor/Blutility/Private/BlutilityMenuExtensions.h
+++ b/Engine/Source/Editor/Blutility/Private/BlutilityMenuExtensions.h
@@ -15,7 +15,7 @@ class FBlutilityMenuExtensions
 {
 public:
 	/** Helper function to get all Blutility classes derived from the specified class name */
-	static void GetBlutilityClasses(TArray<FAssetData>& OutAssets, const FName& InClassName);
+	static void GetBlutilityClasses(TArray<FAssetData>& OutAssets, FTopLevelAssetPath InClassName);
 
 	/** Helper function that populates a menu based on the exposed functions in a set of Blutility objects */
 	static void CreateAssetBlutilityActionsMenu(FMenuBuilder& MenuBuilder, TMap<class IEditorUtilityExtension*, TSet<int32>> Utils, const TArray<FAssetData> SelectedSupportedAssets);
diff --git a/Engine/Source/Editor/Blutility/Private/EditorUtilitySubsystem.cpp b/Engine/Source/Editor/Blutility/Private/EditorUtilitySubsystem.cpp
--- a/Engine/Source/Editor/Blutility/Private/EditorUtilitySubsystem.cpp
+++ b/Engine/Source/Editor/Blutility/Private/EditorUtilitySubsystem.cpp
@@ -445,14 +445,7 @@ UClass* UEditorUtilitySubsystem::FindClassByName(const FString& RawTargetName)
 	UClass* ResultClass = nullptr;
 	if (bIsValidClassName)
 	{
-		if (FPackageName::IsShortPackageName(TargetName))
-		{
-			ResultClass = FindObject<UClass>(ANY_PACKAGE, *TargetName);
-		}
-		else
-		{
-			ResultClass = FindObject<UClass>(nullptr, *TargetName);
-		}
+		ResultClass = UClass::TryFindTypeSlow<UClass>(TargetName);
 	}
 
 	// If we still haven't found anything yet, try the asset registry for blueprints that match the requirements
@@ -478,7 +471,7 @@ UClass* UEditorUtilitySubsystem::FindBlueprintClass(const FString& TargetNameRaw
 
 	FARFilter Filter;
 	Filter.bRecursiveClasses = true;
-	Filter.ClassNames.Add(UBlueprintCore::StaticClass()->GetFName());
+	Filter.ClassPaths.Add(UBlueprintCore::StaticClass()->GetClassPathName());
 
 	// We enumerate all assets to find any blueprints who inherit from native classes directly - or
 	// from other blueprints.
diff --git a/Engine/Source/Editor/Cascade/Private/Cascade.cpp b/Engine/Source/Editor/Cascade/Private/Cascade.cpp
--- a/Engine/Source/Editor/Cascade/Private/Cascade.cpp
+++ b/Engine/Source/Editor/Cascade/Private/Cascade.cpp
@@ -4784,7 +4784,7 @@ void FCascade::OnConvertToSeeded()
 		UE_LOG(LogCascade, Log, TEXT("Non-seeded module %s"), *ClassName);
 		// This only works if the seeded version is names <ClassName>_Seeded!!!!
 		FString SeededClassName = ClassName + TEXT("_Seeded");
-		UClass* SeededClass = FindObject<UClass>(ANY_PACKAGE, *SeededClassName);
+		UClass* SeededClass = UClass::TryFindTypeSlow<UClass>(SeededClassName);
 		if (SeededClass != NULL)
 		{
 			// Find the module index
@@ -4986,7 +4986,7 @@ bool FCascade::ConvertAllModulesToSeeded(UParticleSystem* ParticleSystem)
 						UE_LOG(LogCascade, Log, TEXT("Non-seeded module %s"), *ClassName);
 						// This only works if the seeded version is names <ClassName>_Seeded!!!!
 						FString SeededClassName = ClassName + TEXT("_Seeded");
-						UClass* SeededClass = FindObject<UClass>(ANY_PACKAGE, *SeededClassName);
+						UClass* SeededClass = UClass::TryFindTypeSlow<UClass>(SeededClassName);
 						if (SeededClass != NULL)
 						{
 							TArray<FParticleCurvePair> DistCurves;
diff --git a/Engine/Source/Editor/Cascade/Private/CascadeEmitterCanvasClient.cpp b/Engine/Source/Editor/Cascade/Private/CascadeEmitterCanvasClient.cpp
--- a/Engine/Source/Editor/Cascade/Private/CascadeEmitterCanvasClient.cpp
+++ b/Engine/Source/Editor/Cascade/Private/CascadeEmitterCanvasClient.cpp
@@ -1880,7 +1880,7 @@ TSharedRef<SWidget> FCascadeEmitterCanvasClient::BuildMenuWidgetModule()
 					UE_LOG(LogCascade, Log, TEXT("Non-seeded module %s"), *ClassName);
 					// This only works if the seeded version is names <ClassName>_Seeded!!!!
 					FString SeededClassName = ClassName + TEXT("_Seeded");
-					if (FindObject<UClass>(ANY_PACKAGE, *SeededClassName) != NULL)
+					if (UClass::TryFindTypeSlow<UClass>(SeededClassName) != nullptr)
 					{
 						MenuBuilder.BeginSection("CascadeConvertToSeeded");
 						{
diff --git a/Engine/Source/Editor/ClassViewer/Private/ClassViewerFilter.cpp b/Engine/Source/Editor/ClassViewer/Private/ClassViewerFilter.cpp
--- a/Engine/Source/Editor/ClassViewer/Private/ClassViewerFilter.cpp
+++ b/Engine/Source/Editor/ClassViewer/Private/ClassViewerFilter.cpp
@@ -569,7 +569,7 @@ bool FClassViewerFilter::IsUnloadedClassAllowed(const FClassViewerInitialization
 	// Determine if we allow any developer folder classes, if so determine if this class is in one of the allowed developer folders.
 	static const FString DeveloperPathWithSlash = FPackageName::FilenameToLongPackageName(FPaths::GameDevelopersDir());
 	static const FString UserDeveloperPathWithSlash = FPackageName::FilenameToLongPackageName(FPaths::GameUserDeveloperDir());
-	FString GeneratedClassPathString = InUnloadedClassData->GetClassPath().ToString();
+	FString GeneratedClassPathString = InUnloadedClassData->GetClassPathName().ToString();
 
 	bool bPassesDeveloperFilter = true;
 
diff --git a/Engine/Source/Editor/ClassViewer/Private/ClassViewerNode.cpp b/Engine/Source/Editor/ClassViewer/Private/ClassViewerNode.cpp
--- a/Engine/Source/Editor/ClassViewer/Private/ClassViewerNode.cpp
+++ b/Engine/Source/Editor/ClassViewer/Private/ClassViewerNode.cpp
@@ -13,11 +13,11 @@ FClassViewerNode::FClassViewerNode(UClass* InClass)
 	Class = InClass;
 	ClassName = MakeShareable(new FString(Class->GetName()));
 	ClassDisplayName = MakeShareable(new FString(Class->GetDisplayNameText().ToString()));
-	ClassPath = FName(*Class->GetPathName());
+	ClassPath = Class->GetPathName();
 
 	if (Class->GetSuperClass())
 	{
-		ParentClassPath = FName(*Class->GetSuperClass()->GetPathName());
+		ParentClassPath = Class->GetSuperClass()->GetClassPathName();
 	}
 
 	if (Class->ClassGeneratedBy && Class->ClassGeneratedBy->IsA(UBlueprint::StaticClass()))
diff --git a/Engine/Source/Editor/ClassViewer/Private/ClassViewerNode.h b/Engine/Source/Editor/ClassViewer/Private/ClassViewerNode.h
--- a/Engine/Source/Editor/ClassViewer/Private/ClassViewerNode.h
+++ b/Engine/Source/Editor/ClassViewer/Private/ClassViewerNode.h
@@ -90,10 +90,10 @@ public:
 	TWeakObjectPtr<UBlueprint> Blueprint;
 
 	/** Full object path to the class including _C, set for both blueprint and native */
-	FName ClassPath;
+	FTopLevelAssetPath ClassPath;
 
 	/** Full object path to the parent class, may be blueprint or native */
-	FName ParentClassPath;
+	FTopLevelAssetPath ParentClassPath;
 
 	/** Full path to the Blueprint that this class is loaded from, none for native classes*/
 	FName BlueprintAssetPath;
