Fixing ImportAssetTasks() not properly adding imported asset tot the FAssetImportTask results.

diff --git a/Engine/Source/Developer/AssetTools/Private/AssetTools.cpp b/Engine/Source/Developer/AssetTools/Private/AssetTools.cpp
--- a/Engine/Source/Developer/AssetTools/Private/AssetTools.cpp
+++ b/Engine/Source/Developer/AssetTools/Private/AssetTools.cpp
@@ -1226,10 +1226,10 @@ void UAssetToolsImpl::ImportAssetTasks(const TArray<UAssetImportTask*>& ImportTa
 		Params.bAutomated = ImportTask->bAutomated;
 		Params.SpecifiedFactory = TStrongObjectPtr<UFactory>(ImportTask->Factory);
 		Filenames[0] = ImportTask->Filename;
-		TArray<UObject*> ImportedObjects = ImportAssetsInternal(Filenames, ImportTask->DestinationPath, nullptr, Params);
+		ImportTask->Result = ImportAssetsInternal(Filenames, ImportTask->DestinationPath, nullptr, Params);
 
-		PackagesToSave.Reset(1); 
-		for (UObject* Object : ImportedObjects)
+		PackagesToSave.Reset(ImportTask->Result.Num());
+		for (UObject* Object : ImportTask->Result)
 		{
 			ImportTask->ImportedObjectPaths.Add(Object->GetPathName());
 			if (ImportTask->bSave)
