Remove Create menu's "Blendables" category

diff --git a/Engine/Source/Developer/AssetTools/Private/AssetTools.cpp b/Engine/Source/Developer/AssetTools/Private/AssetTools.cpp
--- a/Engine/Source/Developer/AssetTools/Private/AssetTools.cpp
+++ b/Engine/Source/Developer/AssetTools/Private/AssetTools.cpp
@@ -240,7 +240,6 @@ UAssetToolsImpl::UAssetToolsImpl(const FObjectInitializer& ObjectInitializer)
 	AllocatedCategoryBits.Add(TEXT("_BuiltIn_8"), FAdvancedAssetCategory(EAssetTypeCategories::Media, LOCTEXT("MediaAssetCategory", "Media")));
 	AllocatedCategoryBits.Add(TEXT("_BuiltIn_9"), FAdvancedAssetCategory(EAssetTypeCategories::Textures, LOCTEXT("TextureAssetCategory", "Textures")));
 
-	EAssetTypeCategories::Type BlendablesCategoryBit = RegisterAdvancedAssetCategory(FName(TEXT("Blendables")), LOCTEXT("BlendablesAssetCategory", "Blendables"));
 	EAssetTypeCategories::Type FoliageCategoryBit = RegisterAdvancedAssetCategory(FName(TEXT("Foliage")), LOCTEXT("FoliageAssetCategory", "Foliage"));
 
 	// Register the built-in asset type actions
@@ -291,14 +290,14 @@ UAssetToolsImpl::UAssetToolsImpl(const FObjectInitializer& ObjectInitializer)
 	RegisterAssetTypeActions(MakeShareable(new FAssetTypeActions_LandscapeLayer));
 	RegisterAssetTypeActions(MakeShareable(new FAssetTypeActions_LandscapeGrassType(FoliageCategoryBit)));
 	RegisterAssetTypeActions(MakeShareable(new FAssetTypeActions_LightWeightInstance));
-	RegisterAssetTypeActions(MakeShareable(new FAssetTypeActions_Material(BlendablesCategoryBit)));
+	RegisterAssetTypeActions(MakeShareable(new FAssetTypeActions_Material(EAssetTypeCategories::Type::None)));
 	RegisterAssetTypeActions(MakeShareable(new FAssetTypeActions_MaterialFunction));
 	RegisterAssetTypeActions(MakeShareable(new FAssetTypeActions_MaterialFunctionLayer));
 	RegisterAssetTypeActions(MakeShareable(new FAssetTypeActions_MaterialFunctionLayerInstance));
 	RegisterAssetTypeActions(MakeShareable(new FAssetTypeActions_MaterialFunctionLayerBlend));
 	RegisterAssetTypeActions(MakeShareable(new FAssetTypeActions_MaterialFunctionLayerBlendInstance));
 	RegisterAssetTypeActions(MakeShareable(new FAssetTypeActions_MaterialFunctionInstance));
-	RegisterAssetTypeActions(MakeShareable(new FAssetTypeActions_MaterialInstanceConstant(BlendablesCategoryBit)));
+	RegisterAssetTypeActions(MakeShareable(new FAssetTypeActions_MaterialInstanceConstant(EAssetTypeCategories::Type::None)));
 	RegisterAssetTypeActions(MakeShareable(new FAssetTypeActions_MaterialInstanceDynamic));
 	RegisterAssetTypeActions(MakeShareable(new FAssetTypeActions_MaterialInterface));
 	RegisterAssetTypeActions(MakeShareable(new FAssetTypeActions_MaterialParameterCollection));
