Interchange support for UAssetImportTask
- Don't skip Interchange for APIs that expect immediate results, instead wait until Interchange async import is fully complete and add the results to the regular return path.

diff --git a/Engine/Source/Developer/AssetTools/Private/AssetTools.cpp b/Engine/Source/Developer/AssetTools/Private/AssetTools.cpp
--- a/Engine/Source/Developer/AssetTools/Private/AssetTools.cpp
+++ b/Engine/Source/Developer/AssetTools/Private/AssetTools.cpp
@@ -1771,9 +1771,6 @@ TArray<UObject*> UAssetToolsImpl::ImportAssetsInternal(const TArray<FString>& Fi
 	bUseInterchangeFramework |= bUseInterchangeFrameworkForTextureOnly;
 #endif
 
-	// Block interchange use if the user is not aware that is import can be async. Otherwise, we don't return the imported object and we can't mimic the SpecifiedFactory settings.
-	bUseInterchangeFramework &= Params.bAllowAsyncImport;
-
 	if (!bUseInterchangeFramework && ValidFiles.Num() > 1)
 	{	
 		//Always allow user to cancel the import task if they are importing multiple ValidFiles.
@@ -1812,7 +1809,7 @@ TArray<UObject*> UAssetToolsImpl::ImportAssetsInternal(const TArray<FString>& Fi
 			TArray<FString> FactoryExtensions;
 			Factory->GetSupportedFileExtensions(FactoryExtensions);
 
-			for(auto& FileDest : FilesAndDestinations)
+			for(const TPair<FString, FString>& FileDest : FilesAndDestinations)
 			{
 				const FString FileExtension = FPaths::GetExtension(FileDest.Key);
 
@@ -1895,7 +1892,7 @@ TArray<UObject*> UAssetToolsImpl::ImportAssetsInternal(const TArray<FString>& Fi
 	}
 
 	// We need to sort the factories so that they get tested in priority order
-	for(auto& ExtensionToFactories : ExtensionToFactoriesMap)
+	for(TPair<FString, TArray<UFactory*>>& ExtensionToFactories : ExtensionToFactoriesMap)
 	{
 		ExtensionToFactories.Value.Sort(&UFactory::SortFactoriesByPriority);
 	}
@@ -2058,13 +2055,34 @@ TArray<UObject*> UAssetToolsImpl::ImportAssetsInternal(const TArray<FString>& Fi
 					TPair<UE::Interchange::FAssetImportResultRef, UE::Interchange::FSceneImportResultRef> InterchangeResults =
 						InterchangeManager.ImportSceneAsync(DestinationPath, ScopedSourceData.GetSourceData(), ImportAssetParameters);
 
-					InterchangeResults.Key->OnDone(AppendImportResult);;
-					InterchangeResults.Value->OnDone(AppendAndBroadcastImportResultIfNeeded);
+					InterchangeResults.Get<0>()->OnDone(AppendImportResult);;
+					InterchangeResults.Get<1>()->OnDone(AppendAndBroadcastImportResultIfNeeded);
+
+					if (!Params.bAllowAsyncImport)
+					{
+						InterchangeResults.Get<0>()->WaitUntilDone();
+						InterchangeResults.Get<1>()->WaitUntilDone();
+					}
 				}
 				else
 				{
 					UE::Interchange::FAssetImportResultRef InterchangeResult = (InterchangeManager.ImportAssetAsync(DestinationPath, ScopedSourceData.GetSourceData(), ImportAssetParameters));
 					InterchangeResult->OnDone(AppendAndBroadcastImportResultIfNeeded);
+
+					if (!Params.bAllowAsyncImport)
+					{
+						InterchangeResult->WaitUntilDone();
+					}
+				}
+
+				if (!Params.bAllowAsyncImport)
+				{
+					ReturnObjects.Reserve(ImportStatus->ImportedObjects.Num());
+
+					for (const TWeakObjectPtr<UObject>& ImportedObject : ImportStatus->ImportedObjects)
+					{
+						ReturnObjects.Add(ImportedObject.Get());
+					}
 				}
 
 				//Import done, iterate the next file and destination
@@ -2400,9 +2418,9 @@ TArray<UObject*> UAssetToolsImpl::ImportAssetsInternal(const TArray<FString>& Fi
 	}
 
 	// Clean up and remove the factories we created from the root set
-	for(auto ExtensionIt = ExtensionToFactoriesMap.CreateConstIterator(); ExtensionIt; ++ExtensionIt)
+	for(TMap<FString, TArray<UFactory*>>::TConstIterator ExtensionIt = ExtensionToFactoriesMap.CreateConstIterator(); ExtensionIt; ++ExtensionIt)
 	{
-		for(auto FactoryIt = ExtensionIt.Value().CreateConstIterator(); FactoryIt; ++FactoryIt)
+		for(TArray<UFactory*>::TConstIterator FactoryIt = ExtensionIt.Value().CreateConstIterator(); FactoryIt; ++FactoryIt)
 		{
 			(*FactoryIt)->CleanUp();
 			(*FactoryIt)->RemoveFromRoot();
