Remove the categories from the Asset Type Actions for Blueprint Generated Classes, so that it doesn't show up as a filter in the Content Browser/Reference Viewer

diff --git a/Engine/Source/Developer/AssetTools/Public/AssetTypeActions/AssetTypeActions_BlueprintGeneratedClass.h b/Engine/Source/Developer/AssetTools/Public/AssetTypeActions/AssetTypeActions_BlueprintGeneratedClass.h
--- a/Engine/Source/Developer/AssetTools/Public/AssetTypeActions/AssetTypeActions_BlueprintGeneratedClass.h
+++ b/Engine/Source/Developer/AssetTools/Public/AssetTypeActions/AssetTypeActions_BlueprintGeneratedClass.h
@@ -20,7 +20,9 @@ public:
 	virtual UClass* GetSupportedClass() const override;
 	virtual bool HasActions(const TArray<UObject*>& InObjects) const override { return true; }
 	virtual void GetActions(const TArray<UObject*>& InObjects, FToolMenuSection& Section) override;
-	virtual uint32 GetCategories() override { return EAssetTypeCategories::Blueprint | EAssetTypeCategories::Basic; }
+
+	// We don't want BPGCs to show up in CB/Ref viewer filters since the BP filters already include them
+	virtual uint32 GetCategories() override { return EAssetTypeCategories::None; }
 	// @todo
 	//virtual class UThumbnailInfo* GetThumbnailInfo(UObject* Asset) const override;
 	// End IAssetTypeActions Implementation
