Add Material Layer Instances and Material Layer Blend Instances as filterable types.

diff --git a/Engine/Source/Developer/AssetTools/Private/AssetTypeActions/AssetTypeActions_MaterialFunctionInstance.h b/Engine/Source/Developer/AssetTools/Private/AssetTypeActions/AssetTypeActions_MaterialFunctionInstance.h
--- a/Engine/Source/Developer/AssetTools/Private/AssetTypeActions/AssetTypeActions_MaterialFunctionInstance.h
+++ b/Engine/Source/Developer/AssetTools/Private/AssetTypeActions/AssetTypeActions_MaterialFunctionInstance.h
@@ -32,6 +32,7 @@ public:
 	virtual FColor GetTypeColor() const override { return FColor(0, 175, 175); }
 	virtual UClass* GetSupportedClass() const override;
 	virtual bool CanFilter() override;
+	virtual uint32 GetCategories() override { return EAssetTypeCategories::Materials; }
 private:
 	virtual void ExecuteNewMFI(TArray<TWeakObjectPtr<UMaterialFunctionInterface>> Objects) override;
 	virtual FText GetInstanceText() const override { return NSLOCTEXT("AssetTypeActions", "Material_NewMLI", "Create Layer Instance"); }
@@ -44,6 +45,7 @@ public:
 	virtual FText GetName() const override { return NSLOCTEXT("AssetTypeActions", "AssetTypeActions_MaterialFunctionMaterialLayerBlendInstance", "Material Layer Blend Instance"); }
 	virtual FColor GetTypeColor() const override { return FColor(0, 175, 175); }
 	virtual UClass* GetSupportedClass() const override;
+	virtual uint32 GetCategories() override { return EAssetTypeCategories::Materials; }
 	virtual bool CanFilter() override;
 
 private:
