Fixed incorrect cook path length for plugins mounted from outside of the engine/project directory, and removed a bunch of temporary FString instances

#rnx

diff --git a/Engine/Source/Developer/AssetTools/Private/AssetViewUtils.cpp b/Engine/Source/Developer/AssetTools/Private/AssetViewUtils.cpp
--- a/Engine/Source/Developer/AssetTools/Private/AssetViewUtils.cpp
+++ b/Engine/Source/Developer/AssetTools/Private/AssetViewUtils.cpp
@@ -958,6 +958,21 @@ bool AssetViewUtils::IsDevelopersFolder( const FStringView InPath )
 }
 
 bool AssetViewUtils::IsPluginFolder(const FStringView InPath, EPluginLoadedFrom* OutPluginSource)
+{
+	if (TSharedPtr<IPlugin> Plugin = GetPluginForFolder(InPath))
+	{
+		if (OutPluginSource)
+		{
+			*OutPluginSource = Plugin->GetLoadedFrom();
+		}
+
+		return true;
+	}
+
+	return false;
+}
+
+TSharedPtr<IPlugin> AssetViewUtils::GetPluginForFolder(const FStringView InPath)
 {
 	FStringView PluginName(InPath);
 	if (PluginName.StartsWith(TEXT('/')))
@@ -975,16 +990,11 @@ bool AssetViewUtils::IsPluginFolder(const FStringView InPath, EPluginLoadedFrom*
 	{
 		if (Plugin->IsEnabled() && Plugin->CanContainContent())
 		{
-			if (OutPluginSource != nullptr)
-			{
-				*OutPluginSource = Plugin->GetLoadedFrom();
-			}
-
-			return true;
+			return Plugin;
 		}
 	}
 
-	return false;
+	return nullptr;
 }
 
 void AssetViewUtils::GetObjectsInAssetData(const TArray<FAssetData>& AssetList, TArray<UObject*>& OutDroppedObjects, bool bLoadAllExternalObjects)
@@ -1230,8 +1240,11 @@ FLinearColor AssetViewUtils::GetDefaultColor()
 	return FAppStyle::Get().GetSlateColor(FolderColorName).GetSpecifiedColor();
 }
 
-static const auto CVarMaxFullPathLength = 
-	IConsoleManager::Get().RegisterConsoleVariable( TEXT("MaxAssetFullPath"), FPlatformMisc::GetMaxPathLength(), TEXT("Maximum full path name of an asset.") )->AsVariableInt();
+namespace AssetViewUtils
+{
+static const TConsoleVariableData<int32>* CVarMaxFullPathLength = IConsoleManager::Get().RegisterConsoleVariable(TEXT("MaxAssetFullPath"), FPlatformMisc::GetMaxPathLength(), TEXT("Maximum full path name of an asset."))->AsVariableInt();
+static const TConsoleVariableData<FString>* CVarExternalPluginCookedRootPath = IConsoleManager::Get().RegisterConsoleVariable(TEXT("ExternalPluginCookedAssetRootPath"), TEXT(""), TEXT("Root path to use when estimating the cooked path external plugin assets, or empty to use the standard engine/project root."))->AsVariableString();
+}
 
 bool AssetViewUtils::IsValidObjectPathForCreate(const FString& ObjectPath, FText& OutErrorMessage, bool bAllowExistingAsset)
 {
@@ -1385,80 +1398,72 @@ FString AssetViewUtils::GetPackagePathWithinRoot(const FString& PackageName)
 
 int32 AssetViewUtils::GetPackageLengthForCooking(const FString& PackageName, bool IsInternalBuild)
 {
-	// We assume the game name is 20 characters (the maximum allowed) to make sure that content can be ported between projects
-	static const int32 MaxGameNameLen = 20;
-
-	// Pad out the game name to the maximum allowed
-	const FString GameName = FApp::GetProjectName();
-	FString GameNamePadded = GameName;
-	while (GameNamePadded.Len() < MaxGameNameLen)
+	FString RelativePathToAsset;
+	if (!FPackageName::TryConvertLongPackageNameToFilename(PackageName, RelativePathToAsset, FPackageName::GetAssetPackageExtension()))
 	{
-		GameNamePadded += TEXT(" ");
+		return 0;
 	}
 
-	// We use "LinuxArm64Server" below as it's probably the longest platform name, so will also prove that any shorter platform names will validate correctly
-	const TCHAR* LongPlatformName = TEXT("LinuxArm64Server");
-	const FString AbsoluteRootPath = FPaths::ConvertRelativePathToFull(FPaths::RootDir());
-	const FString AbsoluteGamePath = FPaths::ConvertRelativePathToFull(FPaths::ProjectDir());
-	const FString AbsoluteEnginePath = FPaths::ConvertRelativePathToFull(FPaths::EngineDir());
-	const FString AbsoluteEngineCookPath = AbsoluteGamePath / TEXT("Saved") / TEXT("Cooked") / LongPlatformName / TEXT("Engine");
-	const FString AbsoluteGameCookPath = AbsoluteGamePath / TEXT("Saved") / TEXT("Cooked") / LongPlatformName / GameName;
-
-	EPluginLoadedFrom PluginLoadedFrom;
-	const bool bIsPluginAsset = IsPluginFolder(PackageName, &PluginLoadedFrom);
-	const bool bIsEngineAsset = IsEngineFolder(PackageName) || (bIsPluginAsset && PluginLoadedFrom == EPluginLoadedFrom::Engine);
+	const TCHAR* GameName = FApp::GetProjectName();
+
+	const TSharedPtr<IPlugin> PluginContainingAsset = GetPluginForFolder(PackageName);
+	const bool bIsEngineAsset = IsEngineFolder(PackageName) || (PluginContainingAsset && PluginContainingAsset->GetLoadedFrom() == EPluginLoadedFrom::Engine);
 	const bool bIsProjectAsset = !bIsEngineAsset;
 
-	int32 AbsoluteCookPathToAssetLength = 0;
+	// We use "LinuxArm64Server" below as it's probably the longest platform name, so will also prove that any shorter platform names will validate correctly
+	const FString CookSubPath = FPaths::Combine(TEXT("Saved"), TEXT("Cooked"), TEXT("LinuxArm64Server"), bIsEngineAsset ? TEXT("Engine") : GameName, TEXT("")); // Trailing empty entry ensures a trailing slash
+	const FString AbsolutePathToAsset = FPaths::ConvertRelativePathToFull(RelativePathToAsset);
 
-	FString RelativePathToAsset;
+	FString AbsoluteTargetPath = FPaths::ConvertRelativePathToFull(bIsEngineAsset ? FPaths::EngineDir() : FPaths::ProjectDir());
 
-	const FString AbsolutePath = bIsEngineAsset ? AbsoluteEnginePath : AbsoluteGamePath;
+	int32 AssetPathRelativeToCookRootLen = AbsolutePathToAsset.Len();
+	if (AbsolutePathToAsset.StartsWith(AbsoluteTargetPath, ESearchCase::CaseSensitive))
+	{
+		AssetPathRelativeToCookRootLen -= AbsoluteTargetPath.Len();
+	}
+	else if (ensureMsgf(PluginContainingAsset, TEXT("Only plugins can exist outside of the expected target path of '%s'. '%s' will not calculate an accurate result!"), *AbsoluteTargetPath, *AbsolutePathToAsset))
+	{
+		const FString AbsolutePluginRootPath = FPaths::ConvertRelativePathToFull(PluginContainingAsset->GetBaseDir());
+		if (ensure(AbsolutePathToAsset.StartsWith(AbsolutePluginRootPath, ESearchCase::CaseSensitive)))
+		{
+			AssetPathRelativeToCookRootLen -= AbsolutePluginRootPath.Len();
+			AssetPathRelativeToCookRootLen += FCString::Strlen(TEXT("Plugins/GameFeatures/XXXXXXXX-XXXX-XXXX-XXXX-XXXXXXXXXXXX")); // Use a GUID instead of the plugin name, as some external plugins cook as a GUID
+		}
 
-	const FString& AbsoluteCookPath = bIsEngineAsset ? AbsoluteEngineCookPath : AbsoluteGameCookPath;
+		if (FString ExternalPluginCookedRootPath = CVarExternalPluginCookedRootPath->GetValueOnGameThread(); ExternalPluginCookedRootPath.Len() > 0)
+		{
+			IsInternalBuild = false;
+			AbsoluteTargetPath = MoveTemp(ExternalPluginCookedRootPath);
+			AbsoluteTargetPath /= FString();
+		}
+	}
 
-	if(FPackageName::TryConvertLongPackageNameToFilename(PackageName, RelativePathToAsset, FPackageName::GetAssetPackageExtension()))
+	if (IsInternalBuild)
 	{
-		const FString AbsolutePathToAsset = FPaths::ConvertRelativePathToFull(RelativePathToAsset);
-
-		FString AssetPathWithinCookDir = AbsolutePathToAsset;
-		FPaths::RemoveDuplicateSlashes(AssetPathWithinCookDir);
-		AssetPathWithinCookDir.RemoveFromStart(AbsolutePath, ESearchCase::CaseSensitive);
+		// We assume a constant size for the build machine base path for things that reside within the UE source tree
+		const FString AbsoluteUERootPath = FPaths::ConvertRelativePathToFull(FPaths::RootDir());
+		if (AbsoluteTargetPath.StartsWith(AbsoluteUERootPath, ESearchCase::CaseSensitive))
+		{
+			// Project is within the UE source tree, so remove the UE root path
+			int32 CookPathRelativeToTargetRootLen = CookSubPath.Len();
+			CookPathRelativeToTargetRootLen += (AbsoluteTargetPath.Len() - AbsoluteUERootPath.Len());
 
+			int32 InternalCookPathLen = FCString::Strlen(TEXT("D:/BuildFarm/buildmachine_++depot+UE-Releases+XX.XX/")) + CookPathRelativeToTargetRootLen + AssetPathRelativeToCookRootLen;
 
-		if (IsInternalBuild)
-		{
-			// We assume a constant size for the build machine base path, so strip either the root or game path from the start
-			// (depending on whether the project is part of the main UE source tree or located elsewhere)
-			FString CookDirWithoutBasePath = AbsoluteCookPath;
-			if (CookDirWithoutBasePath.StartsWith(AbsoluteRootPath, ESearchCase::CaseSensitive))
+			// Only add game name padding to project plugins so that they can ported to other projects
+			if (PluginContainingAsset && bIsProjectAsset)
 			{
-				CookDirWithoutBasePath.RemoveFromStart(AbsoluteRootPath, ESearchCase::CaseSensitive);
-			}
-			else
-			{
-				CookDirWithoutBasePath.RemoveFromStart(AbsoluteCookPath, ESearchCase::CaseSensitive);
-			}
-			
-			FString AbsoluteBuildMachineCookPathToAsset = FString(TEXT("D:/BuildFarm/buildmachine_++depot+UE-Releases+XX.XX")) / CookDirWithoutBasePath / AssetPathWithinCookDir;
-
-			// only add game name padding to project plugins so that they can ported to other projects
-			if (bIsPluginAsset && bIsProjectAsset)			{
-				AbsoluteBuildMachineCookPathToAsset.ReplaceInline(*GameName, *GameNamePadded, ESearchCase::CaseSensitive);
+				// We assume the game name is 20 characters (the maximum allowed) to make sure that content can be ported between projects
+				constexpr int32 MaxGameNameLen = 20;
+				InternalCookPathLen += FMath::Max(0, MaxGameNameLen - FCString::Strlen(GameName));
 			}
 
-			AbsoluteCookPathToAssetLength = AbsoluteBuildMachineCookPathToAsset.Len();
-		}
-		else
-		{
-			// Test that the package can be cooked based on the current project path
-			FString AbsoluteCookPathToAsset = AbsoluteCookPath / AssetPathWithinCookDir;
-
-			AbsoluteCookPathToAssetLength = AbsoluteCookPathToAsset.Len();
+			return InternalCookPathLen;
 		}
 	}
-
-	return AbsoluteCookPathToAssetLength;
+	
+	// Test that the package can be cooked based on the current project path
+	return AbsoluteTargetPath.Len() + CookSubPath.Len() + AssetPathRelativeToCookRootLen;
 }
 
 bool AssetViewUtils::IsValidPackageForCooking(const FString& PackageName, FText& OutErrorMessage)
diff --git a/Engine/Source/Developer/AssetTools/Public/AssetViewUtils.h b/Engine/Source/Developer/AssetTools/Public/AssetViewUtils.h
--- a/Engine/Source/Developer/AssetTools/Public/AssetViewUtils.h
+++ b/Engine/Source/Developer/AssetTools/Public/AssetViewUtils.h
@@ -125,6 +125,9 @@ namespace AssetViewUtils
 
 	/** Returns true if the passed-in path is a plugin folder, optionally reporting where the plugin was loaded from */
 	ASSETTOOLS_API bool IsPluginFolder(const FStringView InPath, EPluginLoadedFrom* OutPluginSource = nullptr);
+	
+	/** If the passed-in path is a plugin folder, then return its associated plugin */
+	ASSETTOOLS_API TSharedPtr<IPlugin> GetPluginForFolder(const FStringView InPath);
 
 	/** Get all the objects in a list of asset data */
 	ASSETTOOLS_API void GetObjectsInAssetData(const TArray<FAssetData>& AssetList, TArray<UObject*>& OutDroppedObjects);
