Updated warning message when attempting to rename CDO-referenced assets, for brevity and clarity

#rnx

diff --git a/Engine/Source/Developer/AssetTools/Private/AssetRenameManager.cpp b/Engine/Source/Developer/AssetTools/Private/AssetRenameManager.cpp
--- a/Engine/Source/Developer/AssetTools/Private/AssetRenameManager.cpp
+++ b/Engine/Source/Developer/AssetTools/Private/AssetRenameManager.cpp
@@ -395,23 +395,12 @@ bool FAssetRenameManager::FixReferencesAndRename(const TArray<FAssetRenameData>&
 			UObject* Asset = (*SoftRefIt)->Asset.Get();
 			if (Asset)
 			{
-				FString OptionalTagsString;
-
-				if ((*SoftRefIt)->bWarnAboutProjectSettingsReference)
-				{
-					OptionalTagsString = LOCTEXT("ProjSettingsReferenceTag", "project settings soft reference").ToString();
-				}
-				else
-				{
-					OptionalTagsString = LOCTEXT("SoftReferenceTag", "soft reference").ToString();
-				}
-				
-				AssetNames += FString::Printf(TEXT("\n%s (%s)"), *Asset->GetName(), *OptionalTagsString);
+				AssetNames += FString("\n") + Asset->GetName();
 			}
 		}
 
-		const FText MessageText = FText::Format(LOCTEXT("RenameCDOReferences", "The following assets are referenced by one or more Class Default Objects: \n{0}\n\nContinuing with the rename may require changes to fix references in code or project settings. Assets could otherwise be missing in cooked/packaged builds. Do you wish to continue?"), FText::FromString(AssetNames));
-		if (FMessageDialog::Open(EAppMsgType::YesNo, EAppReturnType::No, MessageText) == EAppReturnType::No)
+		const FText MessageText = FText::Format(LOCTEXT("RenameCDOReferences", "Source code, config INI, and text files may need Find/Replace for:\n\n{0}\n\nOtherwise assets can be missing from cooked builds. Continue with rename?"), FText::FromString(AssetNames));
+		if (FMessageDialog::Open(EAppMsgType::OkCancel, EAppReturnType::Cancel, MessageText) == EAppReturnType::Cancel)
 		{
 			return false;
 		}
@@ -736,6 +725,12 @@ void FAssetRenameManager::FindCDOReferences(const TArrayView<FAssetRenameDataWit
 						AssetToRename->bCreateRedirector |= bSetRedirectorFlags;
 						AssetToRename->bWarnAboutProjectSettingsReference |= bWorkingOnGameMapsSettings;
 
+						if (bWorkingOnGameMapsSettings)
+						{
+							
+							UE_LOG(LogTemp, Warning, TEXT("Config output: %s"), *Cls->GetConfigName());
+						}
+
 						OutSoftReferences.Push(AssetToRename);
 						RemainingSoftRefAssetChecklist.Remove(AssetToRename);
 
