Removed accidental debug logging and no-longer-needed flags from asset renaming operation.

#rnx

diff --git a/Engine/Source/Developer/AssetTools/Private/AssetRenameManager.cpp b/Engine/Source/Developer/AssetTools/Private/AssetRenameManager.cpp
--- a/Engine/Source/Developer/AssetTools/Private/AssetRenameManager.cpp
+++ b/Engine/Source/Developer/AssetTools/Private/AssetRenameManager.cpp
@@ -94,13 +94,11 @@ struct FAssetRenameDataWithReferencers : public FAssetRenameData
 	FText FailureReason;
 	bool bCreateRedirector;
 	bool bRenameFailed;
-	bool bWarnAboutProjectSettingsReference;
 
 	FAssetRenameDataWithReferencers(const FAssetRenameData& InRenameData)
 		: FAssetRenameData(InRenameData)
 		, bCreateRedirector(false)
 		, bRenameFailed(false)
-		, bWarnAboutProjectSettingsReference(false)
 	{
 		if (Asset.IsValid() && !OldObjectPath.IsValid())
 		{
@@ -653,7 +651,6 @@ void FAssetRenameManager::FindCDOReferences(const TArrayView<FAssetRenameDataWit
 	{
 		UClass* Cls = (*ClassDefaultObjectIt);
 		UObject* CDO = Cls->ClassDefaultObject;
-		bool bWorkingOnGameMapsSettings = (CDO == GetDefault<UGameMapsSettings>());
 
 		if (!CDO || !CDO->HasAllFlags(RF_ClassDefaultObject) || !IsValidChecked(CDO) || Cls->ClassGeneratedBy != nullptr)
 		{
@@ -723,13 +720,6 @@ void FAssetRenameManager::FindCDOReferences(const TArrayView<FAssetRenameDataWit
 						(!FinalSoftObjPathName.IsNone() && FinalSoftObjPathName == AssetToRename->OldObjectPath.GetAssetPathName()))
 					{
 						AssetToRename->bCreateRedirector |= bSetRedirectorFlags;
-						AssetToRename->bWarnAboutProjectSettingsReference |= bWorkingOnGameMapsSettings;
-
-						if (bWorkingOnGameMapsSettings)
-						{
-							
-							UE_LOG(LogTemp, Warning, TEXT("Config output: %s"), *Cls->GetConfigName());
-						}
 
 						OutSoftReferences.Push(AssetToRename);
 						RemainingSoftRefAssetChecklist.Remove(AssetToRename);
