Adjustments to allow easier advanced copy customization
* SetPackageThatInitiatedCopy() is now called before any customization steps to fix issue with inconsistent values in PackageThatInitiatedCopy. Notably that PackageThatInitiatedCopy could be the name of the previous package initiator during some method calls if the same advanced copy customization was re-used during the editor lifetime
* Made UAdvancedCopyCustomization::GetARFilter() virtual so customization have an opportunity to define additional filtering

diff --git a/Engine/Source/Developer/AssetTools/Public/AdvancedCopyCustomization.h b/Engine/Source/Developer/AssetTools/Public/AdvancedCopyCustomization.h
--- a/Engine/Source/Developer/AssetTools/Public/AdvancedCopyCustomization.h
+++ b/Engine/Source/Developer/AssetTools/Public/AdvancedCopyCustomization.h
@@ -25,7 +25,7 @@ public:
 	}
 
 	/* Returns the ARFilter for this advanced copy */
-	FARFilter GetARFilter() const
+	virtual FARFilter GetARFilter() const
 	{
 		return FilterForExcludingDependencies;
 	}
