Fix invalid enum forward declaration on linux by including .h instead

diff --git a/Engine/Source/Developer/AssetTools/Public/IAssetTools.h b/Engine/Source/Developer/AssetTools/Public/IAssetTools.h
--- a/Engine/Source/Developer/AssetTools/Public/IAssetTools.h
+++ b/Engine/Source/Developer/AssetTools/Public/IAssetTools.h
@@ -13,9 +13,9 @@
 #include "IAssetTypeActions.h"
 #include "AutomatedAssetImportData.h"
 #include "AssetRegistry/ARFilter.h"
+#include "Engine/Blueprint.h"
 #include "IAssetTools.generated.h"
 
-enum EBlueprintType;
 struct FAssetData;
 class IAssetTypeActions;
 class IClassTypeActions;
