Provide an API for genuinely async imports using AssetImportTask

diff --git a/Engine/Source/Developer/AssetTools/Private/AssetTools.cpp b/Engine/Source/Developer/AssetTools/Private/AssetTools.cpp
--- a/Engine/Source/Developer/AssetTools/Private/AssetTools.cpp
+++ b/Engine/Source/Developer/AssetTools/Private/AssetTools.cpp
@@ -1220,10 +1220,13 @@ void UAssetToolsImpl::ImportAssetTasks(const TArray<UAssetImportTask*>& ImportTa
 		SlowTask.EnterProgressFrame(1, FText::Format(LOCTEXT("Import_ImportingFile", "Importing \"{0}\"..."), FText::FromString(FPaths::GetBaseFilename(ImportTask->Filename))));
 
 		Params.AssetImportTask = ImportTask;
+		Params.bAllowAsyncImport = ImportTask->bAsync;
 		Params.bForceOverrideExisting = ImportTask->bReplaceExisting;
 		Params.bAutomated = ImportTask->bAutomated;
 		Params.SpecifiedFactory = TStrongObjectPtr<UFactory>(ImportTask->Factory);
 		Filenames[0] = ImportTask->Filename;
+
+		PRAGMA_DISABLE_DEPRECATION_WARNINGS
 		ImportTask->Result = ImportAssetsInternal(Filenames, ImportTask->DestinationPath, nullptr, Params);
 
 		PackagesToSave.Reset(ImportTask->Result.Num());
@@ -1235,8 +1238,11 @@ void UAssetToolsImpl::ImportAssetTasks(const TArray<UAssetImportTask*>& ImportTa
 				PackagesToSave.AddUnique(Object->GetOutermost());
 			}
 		}
+		PRAGMA_ENABLE_DEPRECATION_WARNINGS
 
-		if (ImportTask->bSave)
+		// If this wasn't an async import through Interchange (and hence does not have valid AsyncResults),
+		// save imported packages here if required.
+		if (!ImportTask->AsyncResults.IsValid() && ImportTask->bSave)
 		{
 			UEditorLoadingAndSavingUtils::SavePackages(PackagesToSave, true);
 		}
@@ -1749,8 +1755,6 @@ TArray<UObject*> UAssetToolsImpl::ImportAssetsInternal(const TArray<FString>& Fi
 	}
 	TMap< FString, TArray<UFactory*> > ExtensionToFactoriesMap;
 
-	FScopedSlowTask SlowTask(ValidFiles.Num(), LOCTEXT("ImportSlowTask", "Importing"));
-
 	bool bUseInterchangeFramework = false;
 	UInterchangeManager& InterchangeManager = UInterchangeManager::GetInterchangeManager();
 #if WITH_EDITOR
@@ -1758,6 +1762,8 @@ TArray<UObject*> UAssetToolsImpl::ImportAssetsInternal(const TArray<FString>& Fi
 	bUseInterchangeFramework = EditorExperimentalSettings->bEnableInterchangeFramework;
 #endif
 
+	FScopedSlowTask SlowTask(ValidFiles.Num(), LOCTEXT("ImportSlowTask", "Importing"), !bUseInterchangeFramework);
+
 	if (!bUseInterchangeFramework && ValidFiles.Num() > 1)
 	{	
 		//Always allow user to cancel the import task if they are importing multiple ValidFiles.
@@ -1984,11 +1990,23 @@ TArray<UObject*> UAssetToolsImpl::ImportAssetsInternal(const TArray<FString>& Fi
 
 				TFunction<void(UE::Interchange::FImportResult&)> AppendImportResult =
 					// Note: ImportStatus captured by value so that the lambda keeps the shared ptr alive
-					[ImportStatus](UE::Interchange::FImportResult& Result)
+					[ImportStatus, ImportTask = Params.AssetImportTask](UE::Interchange::FImportResult& Result)
 					{
 						ImportStatus->InterchangeResultsContainer->Append(Result.GetResults());
 						ImportStatus->ImportedObjects.Append(Result.GetImportedObjects());
 
+						// If imported through an AssetImportTask object, and it's set to save after import, do this here.
+						if (ImportTask && ImportTask->bSave)
+						{
+							TArray<UPackage*> PackagesToSave;
+							for (UObject* Object : Result.GetImportedObjects())
+							{
+								PackagesToSave.Add(Object->GetOutermost());
+							}
+
+							constexpr bool bDirtyOnly = true;
+							UEditorLoadingAndSavingUtils::SavePackages(PackagesToSave, bDirtyOnly);
+						}
 					};
 
 				TFunction<void(UE::Interchange::FImportResult&)> AppendAndBroadcastImportResultIfNeeded =
@@ -2023,6 +2041,12 @@ TArray<UObject*> UAssetToolsImpl::ImportAssetsInternal(const TArray<FString>& Fi
 					TPair<UE::Interchange::FAssetImportResultRef, UE::Interchange::FSceneImportResultRef> InterchangeResults =
 						InterchangeManager.ImportSceneAsync(DestinationPath, ScopedSourceData.GetSourceData(), ImportAssetParameters);
 
+					// If we have an ImportTask, fill out the asynchronous results object here so the caller can see when the results are ready
+					if ( Params.AssetImportTask)
+					{
+						 Params.AssetImportTask->AsyncResults = InterchangeResults.Get<0>();
+					}
+
 					InterchangeResults.Get<0>()->OnDone(AppendImportResult);;
 					InterchangeResults.Get<1>()->OnDone(AppendAndBroadcastImportResultIfNeeded);
 
@@ -2035,6 +2059,13 @@ TArray<UObject*> UAssetToolsImpl::ImportAssetsInternal(const TArray<FString>& Fi
 				else
 				{
 					UE::Interchange::FAssetImportResultRef InterchangeResult = (InterchangeManager.ImportAssetAsync(DestinationPath, ScopedSourceData.GetSourceData(), ImportAssetParameters));
+
+					// If we have an ImportTask, fill out the asynchronous results object here so the caller can see when the results are ready
+					if ( Params.AssetImportTask)
+					{
+						 Params.AssetImportTask->AsyncResults = InterchangeResult;
+					}
+
 					InterchangeResult->OnDone(AppendAndBroadcastImportResultIfNeeded);
 
 					if (!Params.bAllowAsyncImport)
diff --git a/Engine/Source/Editor/UnrealEd/Private/AssetImportTask.cpp b/Engine/Source/Editor/UnrealEd/Private/AssetImportTask.cpp
--- a/Engine/Source/Editor/UnrealEd/Private/AssetImportTask.cpp
+++ b/Engine/Source/Editor/UnrealEd/Private/AssetImportTask.cpp
@@ -8,3 +8,21 @@ UAssetImportTask::UAssetImportTask()
 {
 
 }
+
+const TArray<UObject*>& UAssetImportTask::GetObjects() const
+{
+	if (AsyncResults.IsValid())
+	{
+		AsyncResults->WaitUntilDone();
+		return AsyncResults->GetImportedObjects();
+	}
+
+	PRAGMA_DISABLE_DEPRECATION_WARNINGS
+	return ToRawPtrTArrayUnsafe(Result);
+	PRAGMA_ENABLE_DEPRECATION_WARNINGS
+}
+
+bool UAssetImportTask::IsAsyncImportComplete() const
+{
+	return !AsyncResults.IsValid() || AsyncResults->GetStatus() == UE::Interchange::FImportResult::EStatus::Done;
+}
diff --git a/Engine/Source/Editor/UnrealEd/Private/Factories/Factory.cpp b/Engine/Source/Editor/UnrealEd/Private/Factories/Factory.cpp
--- a/Engine/Source/Editor/UnrealEd/Private/Factories/Factory.cpp
+++ b/Engine/Source/Editor/UnrealEd/Private/Factories/Factory.cpp
@@ -64,18 +64,19 @@ UObject* UFactory::FactoryCreateFile(UClass* InClass, UObject* InParent, FName I
 
 	if (ScriptFactoryCreateFile(Task))
 	{
-		if (Task->Result.Num() == 0)
+		TArray<UObject*> TaskResults = Task->GetObjects();
+		if (TaskResults.Num() == 0)
 		{
 			return nullptr;
 		}
 		else
 		{
-			for (int32 ResultIndex = 1; ResultIndex < Task->Result.Num(); ++ResultIndex)
+			for (int32 ResultIndex = 1; ResultIndex < TaskResults.Num(); ++ResultIndex)
 			{
-				AdditionalImportedObjects.Add(Task->Result[ResultIndex]);
+				AdditionalImportedObjects.Add(TaskResults[ResultIndex]);
 			}
 
-			return Task->Result[0];
+			return TaskResults[0];
 		}
 	}
 
diff --git a/Engine/Source/Editor/UnrealEd/Public/AssetImportTask.h b/Engine/Source/Editor/UnrealEd/Public/AssetImportTask.h
--- a/Engine/Source/Editor/UnrealEd/Public/AssetImportTask.h
+++ b/Engine/Source/Editor/UnrealEd/Public/AssetImportTask.h
@@ -10,7 +10,8 @@
 #include "UObject/ObjectMacros.h"
 #include "UObject/ObjectPtr.h"
 #include "UObject/UObjectGlobals.h"
-
+#include "UObject/Object.h"
+#include "InterchangeManager.h"
 #include "AssetImportTask.generated.h"
 
 DECLARE_LOG_CATEGORY_EXTERN(LogAssetImportTask, Log, All);
@@ -28,6 +29,7 @@ class UNREALED_API UAssetImportTask : public UObject
 public:
 	UAssetImportTask();
 
+// Task Options
 public:
 	/** Filename to import */
 	UPROPERTY(BlueprintReadWrite, Category = "Asset Import Task")
@@ -57,6 +59,10 @@ public:
 	UPROPERTY(BlueprintReadWrite, Category = "Asset Import Task")
 	bool bSave;
 
+	/** Perform the import asynchronously for file formats where async import is available */
+	UPROPERTY(BlueprintReadWrite, Category = "Asset Import Task")
+	bool bAsync;
+
 	/** Optional factory to use */
 	UPROPERTY(BlueprintReadWrite, Category = "Asset Import Task")
 	TObjectPtr<UFactory> Factory;
@@ -65,12 +71,34 @@ public:
 	UPROPERTY(BlueprintReadWrite, Category = "Asset Import Task")
 	TObjectPtr<UObject> Options;
 
+// Task Results
+public:
+	/**
+	 * Get the list of imported objects.
+	 * Note that if the import was asynchronous, this will block until the results are ready.
+	 * To test whether asynchronous results are ready or not, use IsAsyncImportComplete().
+	 */
+	UFUNCTION(BlueprintCallable, Category = "Asset Import Task")
+	const TArray<UObject*>& GetObjects() const;
+
+	/**
+	 * Query whether this asynchronous import task is complete, and the results are ready to read.
+	 * This will always return true in the case of a blocking import.
+	 */
+	UFUNCTION(BlueprintCallable, Category = "Asset Import Task")
+	bool IsAsyncImportComplete() const;
+
 	/** Paths to objects created or updated after import */
 	UPROPERTY(BlueprintReadWrite, Category = "Asset Import Task")
 	TArray<FString> ImportedObjectPaths;
 
 	/** Imported objects */
-	UPROPERTY(BlueprintReadWrite, Category = "Asset Import Task")
+	// Note: after deprecation this will become a private member
+	UE_DEPRECATED(5.1, "Please do not access this member directly; use GetObjects() instead.")
+	UPROPERTY(meta = (DeprecatedProperty, DeprecationMessage = "Please use the GetObjects function instead."))
 	TArray<TObjectPtr<UObject>> Result;
+
+	/** Async results */
+	UE::Interchange::FAssetImportResultPtr AsyncResults;
 };
 
