AssetTools: Don't use Interchange to import when a UFactory is specified.

diff --git a/Engine/Source/Developer/AssetTools/Private/AssetTools.cpp b/Engine/Source/Developer/AssetTools/Private/AssetTools.cpp
--- a/Engine/Source/Developer/AssetTools/Private/AssetTools.cpp
+++ b/Engine/Source/Developer/AssetTools/Private/AssetTools.cpp
@@ -1754,14 +1754,22 @@ TArray<UObject*> UAssetToolsImpl::ImportAssetsInternal(const TArray<FString>& Fi
 		}
 	}
 	TMap< FString, TArray<UFactory*> > ExtensionToFactoriesMap;
-
-	bool bUseInterchangeFramework = false;
 	UInterchangeManager& InterchangeManager = UInterchangeManager::GetInterchangeManager();
+
+	// Use Interchange if it's enabled and we weren't asked to use a specific UFactory
+	const bool bUseInterchangeFramework = [SpecifiedFactory, &InterchangeManager]()
+	{
+		bool bResult = false;
+		
 #if WITH_EDITOR
-	const UEditorExperimentalSettings* EditorExperimentalSettings = GetDefault<UEditorExperimentalSettings>();
-	bUseInterchangeFramework = EditorExperimentalSettings->bEnableInterchangeFramework;
+		const UEditorExperimentalSettings* EditorExperimentalSettings = GetDefault<UEditorExperimentalSettings>();
+		bResult = EditorExperimentalSettings->bEnableInterchangeFramework;
 #endif
 
+		bResult = bResult && (SpecifiedFactory == nullptr);
+		return bResult;
+	}();
+
 	FScopedSlowTask SlowTask(ValidFiles.Num(), LOCTEXT("ImportSlowTask", "Importing"), !bUseInterchangeFramework);
 
 	if (!bUseInterchangeFramework && ValidFiles.Num() > 1)
