Fixing localization label names in AssetTypeActions_AnimBlueprint

diff --git a/Engine/Source/Developer/AssetTools/Private/AssetTypeActions/AssetTypeActions_AnimBlueprint.cpp b/Engine/Source/Developer/AssetTools/Private/AssetTypeActions/AssetTypeActions_AnimBlueprint.cpp
--- a/Engine/Source/Developer/AssetTools/Private/AssetTypeActions/AssetTypeActions_AnimBlueprint.cpp
+++ b/Engine/Source/Developer/AssetTools/Private/AssetTypeActions/AssetTypeActions_AnimBlueprint.cpp
@@ -144,8 +144,8 @@ void FAssetTypeActions_AnimBlueprint::GetActions(const TArray<UObject*>& InObjec
 
 		Section.AddMenuEntry(
 			"AnimBlueprint_AssignSkeleton",
-			LOCTEXT("SkeletalMesh_AssignSkeleton", "Assign Skeleton"),
-			LOCTEXT("SkeletalMesh_AssignSkeletonTooltip", "Assigns a skeleton to the selected Animation Blueprint(s)."),
+			LOCTEXT("AnimBlueprint_AssignSkeleton", "Assign Skeleton"),
+			LOCTEXT("AnimBlueprint_AssignSkeletonTooltip", "Assigns a skeleton to the selected Animation Blueprint(s)."),
 			FSlateIcon(FAppStyle::GetAppStyleSetName(), "Persona.AssetActions.AssignSkeleton"),
 			FUIAction(
 				FExecuteAction::CreateSP(this, &FAssetTypeActions_AnimBlueprint::ExecuteAssignSkeleton, AnimBlueprints),
