Fix ensure(!IsInitialized()) calling UWorldPartition::Initialize from AdvancedCopyPackages when the source world is opened in the editor

diff --git a/Engine/Source/Developer/AssetTools/Private/AssetTools.cpp b/Engine/Source/Developer/AssetTools/Private/AssetTools.cpp
--- a/Engine/Source/Developer/AssetTools/Private/AssetTools.cpp
+++ b/Engine/Source/Developer/AssetTools/Private/AssetTools.cpp
@@ -1005,7 +1005,10 @@ bool UAssetToolsImpl::AdvancedCopyPackages(
 							if (UWorldPartition* WorldPartition = World->GetWorldPartition())
 							{
 								check(!WorldPartition->IsStreamingEnabled());
-								WorldPartition->Initialize(World, FTransform::Identity);
+								if (!WorldPartition->IsInitialized())
+								{
+									WorldPartition->Initialize(World, FTransform::Identity);
+								}
 							}
 						}
 
