SourceControl: Fixed packages not getting saved even if bForceAutoSave was 'true', because the 'MoveDialogInfo.bSavePackages' value never got inspected.

diff --git a/Engine/Source/Developer/AssetTools/Private/AssetTools.cpp b/Engine/Source/Developer/AssetTools/Private/AssetTools.cpp
--- a/Engine/Source/Developer/AssetTools/Private/AssetTools.cpp
+++ b/Engine/Source/Developer/AssetTools/Private/AssetTools.cpp
@@ -954,6 +954,7 @@ bool UAssetToolsImpl::AdvancedCopyPackages(
 	{
 		TArray<FString> SuccessfullyCopiedDestinationFiles;
 		TArray<FName> SuccessfullyCopiedSourcePackages;
+		TArray<UPackage*> SuccessfullyCopiedDestinationPackages;
 
 		FDuplicatedObjects DuplicatedObjectsLocal;
 		FDuplicatedObjects& DuplicatedObjectsForEachPackage = OutDuplicatedObjects ? *OutDuplicatedObjects : DuplicatedObjectsLocal;
@@ -964,6 +965,7 @@ bool UAssetToolsImpl::AdvancedCopyPackages(
 
 		SuccessfullyCopiedDestinationFiles.Reserve(SourceAndDestPackages.Num());
 		SuccessfullyCopiedSourcePackages.Reserve(SourceAndDestPackages.Num());
+		SuccessfullyCopiedDestinationPackages.Reserve(SourceAndDestPackages.Num());
 		DuplicatedObjectsForEachPackage.Reserve(SourceAndDestPackages.Num());
 		ExistingObjectSet.Reserve(SourceAndDestPackages.Num());
 		NewObjectSet.Reserve(SourceAndDestPackages.Num());
@@ -988,13 +990,10 @@ bool UAssetToolsImpl::AdvancedCopyPackages(
 					if (ExistingObject)
 					{
 						TSet<UPackage*> ObjectsUserRefusedToFullyLoad;
-						ObjectTools::FMoveDialogInfo MoveDialogInfo;
-						MoveDialogInfo.bOkToAll = bCopyOverAllDestinationOverlaps;
-						// The default value for save packages is true if SCC is enabled because the user can use SCC to revert a change
-						MoveDialogInfo.bSavePackages = ISourceControlModule::Get().IsEnabled() || bForceAutosave;
-						MoveDialogInfo.PGN.GroupName = TEXT("");
-						MoveDialogInfo.PGN.ObjectName = FPaths::GetBaseFilename(DestFilename);
-						MoveDialogInfo.PGN.PackageName = DestFilename;
+						ObjectTools::FPackageGroupName PGN;
+						PGN.GroupName = TEXT("");
+						PGN.ObjectName = FPaths::GetBaseFilename(DestFilename);
+						PGN.PackageName = DestFilename;
 						const bool bShouldPromptForDestinationConflict = !bCopyOverAllDestinationOverlaps;
 						TMap<TSoftObjectPtr<UObject>, TSoftObjectPtr<UObject>> DuplicatedObjects;
 
@@ -1012,13 +1011,14 @@ bool UAssetToolsImpl::AdvancedCopyPackages(
 							}
 						}
 
-						if (UObject* NewObject = ObjectTools::DuplicateSingleObject(ExistingObject, MoveDialogInfo.PGN, ObjectsUserRefusedToFullyLoad, bShouldPromptForDestinationConflict, &DuplicatedObjects))
+						if (UObject* NewObject = ObjectTools::DuplicateSingleObject(ExistingObject, PGN, ObjectsUserRefusedToFullyLoad, bShouldPromptForDestinationConflict, &DuplicatedObjects))
 						{
 							ExistingObjectSet.Add(ExistingObject);
 							NewObjectSet.Add(NewObject);
 							DuplicatedObjectsForEachPackage.Add(MoveTemp(DuplicatedObjects));
 							SuccessfullyCopiedSourcePackages.Add(FName(*PackageName));
 							SuccessfullyCopiedDestinationFiles.Add(DestFilename);
+							SuccessfullyCopiedDestinationPackages.Add(NewObject->GetPackage());
 						}
 					}
 				}
@@ -1067,6 +1067,16 @@ bool UAssetToolsImpl::AdvancedCopyPackages(
 
 		ObjectTools::CompileBlueprintsAfterRefUpdate(NewObjectSet.Array());
 
+		// The default value for save packages is true if SCC is enabled because the user can use SCC to revert a change
+		// The save needs to happen before FMarkForAdd is called on the files.
+		bool bSavePackages = ISourceControlModule::Get().IsEnabled() || bForceAutosave;
+		if (bSavePackages)
+		{
+			const bool bCheckDirty = false;
+			const bool bPromptToSave = false;
+			FEditorFileUtils::PromptForCheckoutAndSave(SuccessfullyCopiedDestinationPackages, bCheckDirty, bPromptToSave);
+		}
+
 		FString SourceControlErrors;
 
 		if (SuccessfullyCopiedDestinationFiles.Num() > 0)
@@ -1077,7 +1087,7 @@ bool UAssetToolsImpl::AdvancedCopyPackages(
 				if (ISourceControlModule::Get().IsEnabled())
 				{
 					ISourceControlProvider& SourceControlProvider = ISourceControlModule::Get().GetProvider();
-					if (SourceControlProvider.Execute(ISourceControlOperation::Create<FMarkForAdd>(), SuccessfullyCopiedDestinationFiles) == ECommandResult::Failed)
+					if (SourceControlProvider.Execute(ISourceControlOperation::Create<FMarkForAdd>(), SourceControlHelpers::PackageFilenames(SuccessfullyCopiedDestinationFiles)) == ECommandResult::Failed)
 					{
 						for (auto FileIt(SuccessfullyCopiedDestinationFiles.CreateConstIterator()); FileIt; FileIt++)
 						{
