FNames containing asset paths are deprecated. FSoftObjectPath or FTopLevelAssetPath should be used instead.
Fixups for AssetTools.

diff --git a/Engine/Source/Developer/AssetTools/Private/AssetRenameManager.cpp b/Engine/Source/Developer/AssetTools/Private/AssetRenameManager.cpp
--- a/Engine/Source/Developer/AssetTools/Private/AssetRenameManager.cpp
+++ b/Engine/Source/Developer/AssetTools/Private/AssetRenameManager.cpp
@@ -581,7 +581,7 @@ struct FSoftObjectPathRenameSerializer : public FArchiveUObject
 		const FString& SubPath = Value.GetSubPathString();
 		for (const TPair<FSoftObjectPath, FSoftObjectPath>& Pair : RedirectorMap)
 		{
-			if (Pair.Key.GetAssetPathName() == Value.GetAssetPathName())
+			if (Pair.Key.GetAssetPath() == Value.GetAssetPath())
 			{
 				// Same asset, fix sub path. Asset will be fixed by normal serializePath call below
 				const FString& CheckSubPath = Pair.Key.GetSubPathString();
@@ -600,7 +600,7 @@ struct FSoftObjectPathRenameSerializer : public FArchiveUObject
 
 						FString NewSubPath(SubPath);
 						NewSubPath.ReplaceInline(*CheckSubPath, *Pair.Value.GetSubPathString());
-						Value = FSoftObjectPath(Pair.Value.GetAssetPathName(), NewSubPath);
+						Value = FSoftObjectPath(Pair.Value.GetAssetPath(), NewSubPath);
 					}
 					break;
 				}
@@ -699,15 +699,15 @@ void FAssetRenameManager::FindCDOReferences(const TArrayView<FAssetRenameDataWit
 				TWeakObjectPtr<UBlueprint> ObjectAsBP = UBlueprint::GetBlueprintFromClass(Cast<UBlueprintGeneratedClass>(Object));
 
 				// Resolve to the redirected asset path if necessary
-				FName FinalSoftObjPathName = SoftRefObjPath.GetAssetPathName();
+				FSoftObjectPath FinalSoftObjPath = SoftRefObjPath.GetWithoutSubPath();
 			
 				if (!Object.IsValid() && SoftRefObjPath.IsValid())
 				{
-					FName RedirObjectPathName = GRedirectCollector.GetAssetPathRedirection(SoftRefObjPath.GetAssetPathName());
+					FSoftObjectPath RedirObjectPath = GRedirectCollector.GetAssetPathRedirection(SoftRefObjPath.GetWithoutSubPath());
 
-					if (RedirObjectPathName != NAME_None && RedirObjectPathName != FinalSoftObjPathName)
+					if (!RedirObjectPath.IsNull() && RedirObjectPath != FinalSoftObjPath)
 					{
-						FinalSoftObjPathName = RedirObjectPathName;
+						FinalSoftObjPath = RedirObjectPath;
 					}
 				}
 
@@ -717,7 +717,7 @@ void FAssetRenameManager::FindCDOReferences(const TArrayView<FAssetRenameDataWit
 					// Look for loaded references, indirect blueprint refs to their generated class counterparts, or path name matching
 					if ((Object == AssetToRename->Asset) ||
 						(ObjectAsBP != nullptr && ObjectAsBP == AssetToRename->Asset) ||
-						(!FinalSoftObjPathName.IsNone() && FinalSoftObjPathName == AssetToRename->OldObjectPath.GetAssetPathName()))
+						(!FinalSoftObjPath.IsNull() && FinalSoftObjPath == AssetToRename->OldObjectPath.GetWithoutSubPath()))
 					{
 						AssetToRename->bCreateRedirector |= bSetRedirectorFlags;
 
@@ -1215,7 +1215,7 @@ void FAssetRenameManager::RenameReferencingSoftObjectPaths(const TArray<UPackage
 	{
 		if (Pair.Key.IsAsset())
 		{
-			GRedirectCollector.AddAssetPathRedirection(Pair.Key.GetAssetPathName(), Pair.Value.GetAssetPathName());
+			GRedirectCollector.AddAssetPathRedirection(Pair.Key.GetWithoutSubPath(), Pair.Value.GetWithoutSubPath());
 		}
 	}
 
@@ -1328,7 +1328,7 @@ bool FAssetRenameManager::CheckPackageForSoftObjectReferences(UPackage* Package,
 			const FString& SubPath = CachedKey.GetSubPathString();
 
 			// Stop as soon as we're not anymore in the range we're searching
-			if (Pair.Key.GetAssetPathName() != CachedKey.GetAssetPathName())
+			if (Pair.Key.GetWithoutSubPath() != CachedKey.GetWithoutSubPath())
 			{
 				break;
 			}
diff --git a/Engine/Source/Developer/AssetTools/Private/AssetTools.cpp b/Engine/Source/Developer/AssetTools/Private/AssetTools.cpp
--- a/Engine/Source/Developer/AssetTools/Private/AssetTools.cpp
+++ b/Engine/Source/Developer/AssetTools/Private/AssetTools.cpp
@@ -1423,7 +1423,7 @@ void UAssetToolsImpl::CreateUniqueAssetName(const FString& InBasePackageName, co
 		// Use the asset registry if possible to find existing assets without loading them
 		if ( !AssetRegistryModule.Get().IsLoadingAssets() )
 		{
-			FAssetData AssetData = AssetRegistryModule.Get().GetAssetByObjectPath(*ObjectPath);
+			FAssetData AssetData = AssetRegistryModule.Get().GetAssetByObjectPath(FSoftObjectPath(ObjectPath));
 			if(AssetData.IsValid())
 			{
 				bObjectExists = true;
diff --git a/Engine/Source/Developer/AssetTools/Private/AssetTypeActions/AssetTypeActions_Class.cpp b/Engine/Source/Developer/AssetTools/Private/AssetTypeActions/AssetTypeActions_Class.cpp
--- a/Engine/Source/Developer/AssetTools/Private/AssetTypeActions/AssetTypeActions_Class.cpp
+++ b/Engine/Source/Developer/AssetTools/Private/AssetTypeActions/AssetTypeActions_Class.cpp
@@ -174,7 +174,7 @@ void FAssetTypeActions_Class::OpenAssetEditor( const TArray<UObject*>& InObjects
 
 TWeakPtr<IClassTypeActions> FAssetTypeActions_Class::GetClassTypeActions(const FAssetData& AssetData) const
 {
-	UClass* Class = FindObject<UClass>(nullptr, *AssetData.ObjectPath.ToString());
+	UClass* Class = FindObject<UClass>(nullptr, *AssetData.GetObjectPathString());
 	if(Class)
 	{
 		FAssetToolsModule& AssetToolsModule = FAssetToolsModule::GetModule();
diff --git a/Engine/Source/Developer/AssetTools/Private/AssetTypeActions/AssetTypeActions_Texture.cpp b/Engine/Source/Developer/AssetTools/Private/AssetTypeActions/AssetTypeActions_Texture.cpp
--- a/Engine/Source/Developer/AssetTools/Private/AssetTypeActions/AssetTypeActions_Texture.cpp
+++ b/Engine/Source/Developer/AssetTools/Private/AssetTypeActions/AssetTypeActions_Texture.cpp
@@ -522,7 +522,7 @@ private:
 				[
 						(DetailedInfoText.IsEmpty())
 						? static_cast<TSharedRef<SWidget>>(SNew(STextBlock)
-							.Text(FText::FromName(Asset.ObjectPath)))
+							.Text(FText::FromString(Asset.GetObjectPathString())))
 					: SNew(SButton)
 					.ButtonStyle(FCoreStyle::Get(), "NoBorder")
 					.VAlign(VAlign_Center)
@@ -551,7 +551,7 @@ private:
 						.VAlign(VAlign_Center)
 						[
 							SNew(STextBlock)
-							.Text(FText::FromName(Asset.ObjectPath))
+							.Text(FText::FromString(Asset.GetObjectPathString()))
 						]
 					]
 			    ]
diff --git a/Engine/Source/Developer/AssetTools/Private/AssetViewUtils.cpp b/Engine/Source/Developer/AssetTools/Private/AssetViewUtils.cpp
--- a/Engine/Source/Developer/AssetTools/Private/AssetViewUtils.cpp
+++ b/Engine/Source/Developer/AssetTools/Private/AssetViewUtils.cpp
@@ -363,7 +363,7 @@ bool AssetViewUtils::DeleteFolders(const TArray<FString>& PathsToDelete)
 		TArray<FString> ObjectPaths;
 		for ( auto AssetIt = AssetDataList.CreateConstIterator(); AssetIt; ++AssetIt )
 		{
-			ObjectPaths.Add((*AssetIt).ObjectPath.ToString());
+			ObjectPaths.Add((*AssetIt).GetObjectPathString());
 		}
 
 		// Load all the assets in the selected paths
@@ -712,7 +712,7 @@ bool AssetViewUtils::PrepareFoldersForDragDrop(const TArray<FString>& SourcePath
 
 		for ( auto AssetIt = AssetDataList.CreateConstIterator(); AssetIt; ++AssetIt )
 		{
-			ObjectPathsToWarnAbout.Add((*AssetIt).ObjectPath.ToString());
+			ObjectPathsToWarnAbout.Add((*AssetIt).GetObjectPathString());
 		}
 	}
 
@@ -729,7 +729,7 @@ bool AssetViewUtils::PrepareFoldersForDragDrop(const TArray<FString>& SourcePath
 		TArray<FString> ObjectPaths;
 		for ( auto AssetIt = AssetDataList.CreateConstIterator(); AssetIt; ++AssetIt )
 		{
-			ObjectPaths.Add((*AssetIt).ObjectPath.ToString());
+			ObjectPaths.Add((*AssetIt).GetObjectPathString());
 		}
 
 		// Load all assets in this path if needed
@@ -1313,7 +1313,7 @@ bool AssetViewUtils::IsValidObjectPathForCreate(const FString& ObjectPath, const
 
 	// Check for an existing asset
 	FAssetRegistryModule& AssetRegistryModule = FModuleManager::LoadModuleChecked<FAssetRegistryModule>("AssetRegistry");
-	FAssetData ExistingAsset = AssetRegistryModule.Get().GetAssetByObjectPath(FName(*ObjectPath));
+	FAssetData ExistingAsset = AssetRegistryModule.Get().GetAssetByObjectPath(FSoftObjectPath(ObjectPath));
 	if (ExistingAsset.IsValid())
 	{
 		// An asset of a different type already exists at this location, inform the user and continue
