Removed unacceptable word from AssetRenameManager.h

diff --git a/Engine/Source/Developer/AssetTools/Private/AssetRenameManager.h b/Engine/Source/Developer/AssetTools/Private/AssetRenameManager.h
--- a/Engine/Source/Developer/AssetTools/Private/AssetRenameManager.h
+++ b/Engine/Source/Developer/AssetTools/Private/AssetRenameManager.h
@@ -67,7 +67,7 @@ private:
 
 	/**
 	 * Get lists of assets with references from CDOs
-	 * @param AssetsToRename      Master list of assets to be renamed
+	 * @param AssetsToRename      List of assets to be renamed
 	 * @param OutReferences       Output array of hard asset references
 	 * @param OutSoftReferences   Output array of soft path asset references
 	 * @param bSetRedirectorFlags If true, each renamed asset with a soft reference will be marked as needing a redirector
