Added filtering for animation blueprint interfaces

diff --git a/Engine/Source/Developer/AssetTools/Private/AssetTools.cpp b/Engine/Source/Developer/AssetTools/Private/AssetTools.cpp
--- a/Engine/Source/Developer/AssetTools/Private/AssetTools.cpp
+++ b/Engine/Source/Developer/AssetTools/Private/AssetTools.cpp
@@ -45,6 +45,7 @@
 #include "AssetTypeActions/AssetTypeActions_VectorField.h"
 #include "AssetTypeActions/AssetTypeActions_AnimationAsset.h"
 #include "AssetTypeActions/AssetTypeActions_AnimBlueprint.h"
+#include "AssetTypeActions/AssetTypeActions_AnimBlueprintInterface.h"
 #include "AssetTypeActions/AssetTypeActions_AnimBlueprintGeneratedClass.h"
 #include "AssetTypeActions/AssetTypeActions_AnimBoneCompressionSettings.h"
 #include "AssetTypeActions/AssetTypeActions_AnimComposite.h"
@@ -248,6 +249,7 @@ UAssetToolsImpl::UAssetToolsImpl(const FObjectInitializer& ObjectInitializer)
 	RegisterAssetTypeActions(MakeShareable(new FAssetTypeActions_ActorFolder));
 	RegisterAssetTypeActions(MakeShareable(new FAssetTypeActions_AnimationAsset));
 	RegisterAssetTypeActions(MakeShareable(new FAssetTypeActions_AnimBlueprint));
+	RegisterAssetTypeActions(MakeShareable(new FAssetTypeActions_AnimBlueprintInterface));
 	RegisterAssetTypeActions(MakeShareable(new FAssetTypeActions_AnimBlueprintGeneratedClass));
 	RegisterAssetTypeActions(MakeShareable(new FAssetTypeActions_AnimBoneCompressionSettings));
 	RegisterAssetTypeActions(MakeShareable(new FAssetTypeActions_AnimComposite));
diff --git a/Engine/Source/Developer/AssetTools/Private/AssetTypeActions/AssetTypeActions_AnimBlueprintInterface.cpp b/Engine/Source/Developer/AssetTools/Private/AssetTypeActions/AssetTypeActions_AnimBlueprintInterface.cpp
new file mode 100644
--- /dev/null
+++ b/Engine/Source/Developer/AssetTools/Private/AssetTypeActions/AssetTypeActions_AnimBlueprintInterface.cpp
@@ -0,0 +1,16 @@
+﻿// Copyright Epic Games, Inc. All Rights Reserved.
+
+#include "AssetTypeActions_AnimBlueprintInterface.h"
+
+void FAssetTypeActions_AnimBlueprintInterface::BuildBackendFilter(FARFilter& InFilter)
+{
+	FAssetTypeActions_AnimBlueprint::BuildBackendFilter(InFilter);
+
+	InFilter.TagsAndValues.Add(GET_MEMBER_NAME_CHECKED(UBlueprint, BlueprintType), FString(TEXT("BPTYPE_Interface")));
+}
+
+FName FAssetTypeActions_AnimBlueprintInterface::GetFilterName() const
+{
+	FName SuperName = FAssetTypeActions_AnimBlueprint::GetFilterName();
+	return *(SuperName.ToString() + TEXT("Interface"));
+}
\ No newline at end of file
diff --git a/Engine/Source/Developer/AssetTools/Private/AssetTypeActions/AssetTypeActions_AnimBlueprintInterface.h b/Engine/Source/Developer/AssetTools/Private/AssetTypeActions/AssetTypeActions_AnimBlueprintInterface.h
new file mode 100644
--- /dev/null
+++ b/Engine/Source/Developer/AssetTools/Private/AssetTypeActions/AssetTypeActions_AnimBlueprintInterface.h
@@ -0,0 +1,16 @@
+﻿// Copyright Epic Games, Inc. All Rights Reserved.
+
+#pragma once
+
+#include "CoreMinimal.h"
+#include "AssetTypeActions_AnimBlueprint.h"
+
+class FAssetTypeActions_AnimBlueprintInterface : public FAssetTypeActions_AnimBlueprint
+{
+public:
+	// IAssetTypeActions Implementation
+	virtual FText GetName() const override { return NSLOCTEXT("AssetTypeActions", "AssetTypeActions_AnimBlueprintInterface", "Animation Blueprint Interface"); }
+	virtual void BuildBackendFilter(FARFilter& InFilter) override;
+	virtual FName GetFilterName() const override;
+	virtual void GetActions(const TArray<UObject*>& InObjects, struct FToolMenuSection& Section) override {}
+};
