Fix animation blueprint interface being displayed in asset creation menu instead of animation blueprints

Also rename animation blueprint interface to animation layer interface

diff --git a/Engine/Source/Developer/AssetTools/Private/AssetTypeActions/AssetTypeActions_AnimBlueprintInterface.h b/Engine/Source/Developer/AssetTools/Private/AssetTypeActions/AssetTypeActions_AnimBlueprintInterface.h
--- a/Engine/Source/Developer/AssetTools/Private/AssetTypeActions/AssetTypeActions_AnimBlueprintInterface.h
+++ b/Engine/Source/Developer/AssetTools/Private/AssetTypeActions/AssetTypeActions_AnimBlueprintInterface.h
@@ -9,7 +9,7 @@ class FAssetTypeActions_AnimBlueprintInterface : public FAssetTypeActions_AnimBl
 {
 public:
 	// IAssetTypeActions Implementation
-	virtual FText GetName() const override { return NSLOCTEXT("AssetTypeActions", "AssetTypeActions_AnimBlueprintInterface", "Animation Blueprint Interface"); }
+	virtual FText GetName() const override { return NSLOCTEXT("AssetTypeActions", "AssetTypeActions_AnimLayerInterface", "Animation Layer Interface"); }
 	virtual void BuildBackendFilter(FARFilter& InFilter) override;
 	virtual FName GetFilterName() const override;
 	virtual void GetActions(const TArray<UObject*>& InObjects, struct FToolMenuSection& Section) override {}
diff --git a/Engine/Source/Editor/UnrealEd/Classes/Factories/AnimBlueprintFactory.h b/Engine/Source/Editor/UnrealEd/Classes/Factories/AnimBlueprintFactory.h
--- a/Engine/Source/Editor/UnrealEd/Classes/Factories/AnimBlueprintFactory.h
+++ b/Engine/Source/Editor/UnrealEd/Classes/Factories/AnimBlueprintFactory.h
@@ -39,6 +39,7 @@ class UAnimBlueprintFactory : public UFactory
 	bool bTemplate;
 	
 	//~ Begin UFactory Interface
+	virtual FText GetDisplayName() const override;
 	virtual bool ConfigureProperties() override;
 	virtual UObject* FactoryCreateNew(UClass* Class, UObject* InParent, FName Name, EObjectFlags Flags, UObject* Context, FFeedbackContext* Warn, FName CallingContext) override;
 	virtual UObject* FactoryCreateNew(UClass* Class,UObject* InParent,FName Name,EObjectFlags Flags,UObject* Context,FFeedbackContext* Warn) override;
diff --git a/Engine/Source/Editor/UnrealEd/Private/Factories/AnimBlueprintFactory.cpp b/Engine/Source/Editor/UnrealEd/Private/Factories/AnimBlueprintFactory.cpp
--- a/Engine/Source/Editor/UnrealEd/Private/Factories/AnimBlueprintFactory.cpp
+++ b/Engine/Source/Editor/UnrealEd/Private/Factories/AnimBlueprintFactory.cpp
@@ -435,6 +435,11 @@ UAnimBlueprintFactory::UAnimBlueprintFactory(const FObjectInitializer& ObjectIni
 	ParentClass = UAnimInstance::StaticClass();
 }
 
+FText UAnimBlueprintFactory::GetDisplayName() const
+{
+	return LOCTEXT("AnimationBlueprintFactoryDescription", "Animation Blueprint");
+}
+
 bool UAnimBlueprintFactory::ConfigureProperties()
 {
 	TSharedRef<SAnimBlueprintCreateDialog> Dialog = SNew(SAnimBlueprintCreateDialog);
