Add new EditorDelegate OnPreAssetDestructiveAction

diff --git a/Engine/Plugins/Editor/ContentBrowser/ContentBrowserAssetDataSource/Source/ContentBrowserAssetDataSource/Private/ContentBrowserAssetDataCore.cpp b/Engine/Plugins/Editor/ContentBrowser/ContentBrowserAssetDataSource/Source/ContentBrowserAssetDataSource/Private/ContentBrowserAssetDataCore.cpp
--- a/Engine/Plugins/Editor/ContentBrowser/ContentBrowserAssetDataSource/Source/ContentBrowserAssetDataSource/Private/ContentBrowserAssetDataCore.cpp
+++ b/Engine/Plugins/Editor/ContentBrowser/ContentBrowserAssetDataSource/Source/ContentBrowserAssetDataSource/Private/ContentBrowserAssetDataCore.cpp
@@ -854,6 +854,15 @@ bool RenameAssetFileItem(IAssetTools* InAssetTools, const FContentBrowserAssetFi
 	// We need to potentially load the asset in order to rename it
 	if (UObject* Asset = InAssetPayload.LoadAsset({ ULevel::LoadAllExternalObjectsTag }))
 	{
+		FResultMessage Result;
+		Result.bSucceeded = true;
+		FEditorDelegates::OnPreDestructiveAssetAction.Broadcast({Asset}, EDestructiveAssetActions::AssetRename, Result);
+		if (!Result.WasSuccesful())
+		{
+			UE_LOG(LogContentBrowserAssetDataSource, Warning, TEXT("%s"), *Result.GetErrorMessage());
+			return false;
+		}
+
 		const FString PackagePath = FPackageName::GetLongPackagePath(Asset->GetOutermost()->GetName());
 
 		TArray<FAssetRenameData> AssetsAndNames;
@@ -1016,6 +1025,15 @@ bool MoveAssetFileItems(TArrayView<const TSharedRef<const FContentBrowserAssetFi
 		}
 	}
 
+	FResultMessage Result;
+	Result.bSucceeded = true;
+	FEditorDelegates::OnPreDestructiveAssetAction.Broadcast(AssetsToMove, EDestructiveAssetActions::AssetMove, Result);
+	if (!Result.WasSuccesful())
+	{
+		UE_LOG(LogContentBrowserAssetDataSource, Warning, TEXT("%s"), *Result.GetErrorMessage());
+		return false;
+	}
+
 	if (AssetsToMove.Num() > 0)
 	{
 		AssetViewUtils::MoveAssets(AssetsToMove, InDestPath.ToString());
diff --git a/Engine/Source/Developer/AssetTools/Private/AssetViewUtils.cpp b/Engine/Source/Developer/AssetTools/Private/AssetViewUtils.cpp
--- a/Engine/Source/Developer/AssetTools/Private/AssetViewUtils.cpp
+++ b/Engine/Source/Developer/AssetTools/Private/AssetViewUtils.cpp
@@ -43,6 +43,8 @@
 
 #define LOCTEXT_NAMESPACE "ContentBrowser"
 
+DEFINE_LOG_CATEGORY_STATIC(LogAssetViewTools, Warning, Warning);
+
 #define MAX_CLASS_NAME_LENGTH 32 // Enforce a reasonable class name length so the path is not too long for FPlatformMisc::GetMaxPathLength()
 
 
@@ -550,6 +552,16 @@ bool AssetViewUtils::RenameFolder(const FString& DestPath, const FString& Source
 	TArray<UObject*> ObjectsInFolder;
 	const bool bLoadAllExternalObjects = true;
 	GetObjectsInAssetData(AssetsInFolder, ObjectsInFolder, bLoadAllExternalObjects);
+
+	FResultMessage Result;
+	Result.bSucceeded = true;
+	FEditorDelegates::OnPreDestructiveAssetAction.Broadcast(ObjectsInFolder, EDestructiveAssetActions::AssetRename, Result);
+	if (!Result.WasSuccesful())
+	{
+		UE_LOG(LogAssetViewTools, Warning, TEXT("%s"), *Result.GetErrorMessage());
+		return false;
+	}
+
 	MoveAssets(ObjectsInFolder, DestPath, SourcePath);
 
 	// Now check to see if the original folder is empty, if so we can delete it
@@ -649,6 +661,21 @@ bool AssetViewUtils::MoveFolders(const TArray<FString>& InSourcePathNames, const
 	{
 		return false;
 	}
+
+	TArray<UObject*> AssetsToMove;
+	for (auto PathIt = SourcePathToLoadedAssets.CreateConstIterator(); PathIt; ++PathIt)
+	{
+		AssetsToMove.Append(PathIt.Value());
+	}
+
+	FResultMessage Result;
+	Result.bSucceeded = true;
+	FEditorDelegates::OnPreDestructiveAssetAction.Broadcast(AssetsToMove, EDestructiveAssetActions::AssetMove, Result);
+	if (!Result.WasSuccesful())
+	{
+		UE_LOG(LogAssetViewTools, Warning, TEXT("%s"), *Result.GetErrorMessage());
+		return false;
+	}
 	
 	// Load the Asset Registry to update paths during the move
 	FAssetRegistryModule& AssetRegistryModule = FModuleManager::LoadModuleChecked<FAssetRegistryModule>("AssetRegistry");
diff --git a/Engine/Source/Editor/UnrealEd/Classes/Editor/UnrealEdTypes.h b/Engine/Source/Editor/UnrealEd/Classes/Editor/UnrealEdTypes.h
--- a/Engine/Source/Editor/UnrealEd/Classes/Editor/UnrealEdTypes.h
+++ b/Engine/Source/Editor/UnrealEd/Classes/Editor/UnrealEdTypes.h
@@ -2,7 +2,7 @@
 
 /**
  *	This will hold all of our enums and types and such that we need to
- *	use in multiple files where the enum can'y be mapped to a specific file.
+ *	use in multiple files where the enum can't be mapped to a specific file.
  */
 
 #pragma once
@@ -129,6 +129,29 @@ enum ELevelViewportType
 	LVT_None = 255,
 };
 
+UENUM()
+enum EDestructiveAssetActions
+{
+	AssetDelete = 0,
+	AssetRename = 1 << 0,
+	AssetMove = 1 << 1,
+	AssetPrivatize = 1 << 2
+};
+
+/** Simple Return struct for supplying success or failure with an optional error message */
+USTRUCT()
+struct FResultMessage
+{
+	GENERATED_USTRUCT_BODY()
+public:
+	bool WasSuccesful() { return bSucceeded; }
+	FString GetErrorMessage() { return ErrorMessage; }
+	void SetErrorMessage(FString& InErrorMessage) { InErrorMessage = ErrorMessage; }
+	bool bSucceeded = false;
+private:
+	FString ErrorMessage;
+};
+
 
 UCLASS(abstract, config=UnrealEd)
 class UUnrealEdTypes : public UObject
diff --git a/Engine/Source/Editor/UnrealEd/Private/Editor.cpp b/Engine/Source/Editor/UnrealEd/Private/Editor.cpp
--- a/Engine/Source/Editor/UnrealEd/Private/Editor.cpp
+++ b/Engine/Source/Editor/UnrealEd/Private/Editor.cpp
@@ -121,6 +121,7 @@ FEditorDelegates::FOnPostSaveExternalActors				FEditorDelegates::PostSaveExterna
 FEditorDelegates::FOnFinishPickingBlueprintClass		FEditorDelegates::OnFinishPickingBlueprintClass;
 FEditorDelegates::FOnNewAssetCreation					FEditorDelegates::OnConfigureNewAssetProperties;
 FEditorDelegates::FOnNewAssetCreation					FEditorDelegates::OnNewAssetCreated;
+FEditorDelegates::FOnPreDestructiveAssetAction          FEditorDelegates::OnPreDestructiveAssetAction;
 FEditorDelegates::FOnAssetPreImport						FEditorDelegates::OnAssetPreImport;
 FEditorDelegates::FOnAssetPostImport					FEditorDelegates::OnAssetPostImport;
 FEditorDelegates::FOnAssetReimport						FEditorDelegates::OnAssetReimport;
diff --git a/Engine/Source/Editor/UnrealEd/Private/ObjectTools.cpp b/Engine/Source/Editor/UnrealEd/Private/ObjectTools.cpp
--- a/Engine/Source/Editor/UnrealEd/Private/ObjectTools.cpp
+++ b/Engine/Source/Editor/UnrealEd/Private/ObjectTools.cpp
@@ -2680,6 +2680,16 @@ namespace ObjectTools
 			return 0;
 		}
 
+		FResultMessage Result;
+		Result.bSucceeded = true;
+		FEditorDelegates::OnPreDestructiveAssetAction.Broadcast(ObjectsToDelete, EDestructiveAssetActions::AssetDelete, Result);
+
+		if (!Result.WasSuccesful())
+		{
+			UE_LOG(LogObjectTools, Warning, TEXT("%s"), *Result.GetErrorMessage());
+			return 0;
+		}
+
 		// Load the asset registry module
 		FAssetRegistryModule& AssetRegistryModule = FModuleManager::LoadModuleChecked<FAssetRegistryModule>(TEXT("AssetRegistry"));
 
@@ -2774,6 +2784,16 @@ namespace ObjectTools
 			return 0;
 		}
 
+		FResultMessage Result;
+		Result.bSucceeded = true;
+		FEditorDelegates::OnPreDestructiveAssetAction.Broadcast(ObjectsToPrivatize, EDestructiveAssetActions::AssetPrivatize, Result);
+
+		if (!Result.WasSuccesful())
+		{
+			UE_LOG(LogObjectTools, Warning, TEXT("%s"), *Result.GetErrorMessage());
+			return 0;
+		}
+
 		FAssetRegistryModule& AssetRegistryModule = FModuleManager::LoadModuleChecked<FAssetRegistryModule>(TEXT("AssetRegistry"));
 
 		if (AssetRegistryModule.Get().IsLoadingAssets())
diff --git a/Engine/Source/Editor/UnrealEd/Public/Editor.h b/Engine/Source/Editor/UnrealEd/Public/Editor.h
--- a/Engine/Source/Editor/UnrealEd/Public/Editor.h
+++ b/Engine/Source/Editor/UnrealEd/Public/Editor.h
@@ -86,6 +86,8 @@ struct UNREALED_API FEditorDelegates
 	DECLARE_MULTICAST_DELEGATE_OneParam(FOnStandaloneLocalPlayEvent, const uint32);
 	/** delegate type for beginning or finishing configuration of the properties of a new asset */
 	DECLARE_MULTICAST_DELEGATE_OneParam(FOnNewAssetCreation, UFactory*);
+	/** delegate for when assets are about to undergo a destructive action caused by the Editor UI (Delete, Rename, Move, Privatize, etc.) */
+	DECLARE_MULTICAST_DELEGATE_ThreeParams(FOnPreDestructiveAssetAction, const TArray<UObject*>&, EDestructiveAssetActions, FResultMessage&);
 	/** delegate type fired when new assets are being (re-)imported. Params: UFactory* InFactory, UClass* InClass, UObject* InParent, const FName& Name, const TCHAR* Type */
 	DECLARE_MULTICAST_DELEGATE_FiveParams(FOnAssetPreImport, UFactory*, UClass*, UObject*, const FName&, const TCHAR*);
 	/** delegate type fired when new assets have been (re-)imported. Note: InCreatedObject can be NULL if import failed. Params: UFactory* InFactory, UObject* InCreatedObject */
@@ -241,6 +243,8 @@ struct UNREALED_API FEditorDelegates
 	static FOnFinishPickingBlueprintClass OnFinishPickingBlueprintClass;
 	/** Called when beginning configuration of a new asset */
 	static FOnNewAssetCreation OnConfigureNewAssetProperties;
+	/** Called when an asset is about to undergo a destructive action caused by the Editor UI (Delete, Move, Rename, Privatize, etc.) */
+	static FOnPreDestructiveAssetAction OnPreDestructiveAssetAction;
 	/** Called when finishing configuration of a new asset */
 	static FOnNewAssetCreation OnNewAssetCreated;
 	/** Called when new assets are being (re-)imported. */
