Expose Asset Migration command to scripting

diff --git a/Engine/Source/Developer/AssetTools/Private/AssetTools.cpp b/Engine/Source/Developer/AssetTools/Private/AssetTools.cpp
--- a/Engine/Source/Developer/AssetTools/Private/AssetTools.cpp
+++ b/Engine/Source/Developer/AssetTools/Private/AssetTools.cpp
@@ -1697,25 +1697,32 @@ bool UAssetToolsImpl::CreateDiffProcess(const FString& DiffCommand,  const FStri
 }
 
 void UAssetToolsImpl::MigratePackages(const TArray<FName>& PackageNamesToMigrate) const
+{
+	FMigrationOptions Options = FMigrationOptions();
+	Options.bPrompt = true;
+	MigratePackages(PackageNamesToMigrate, FString(), Options);
+}
+
+void UAssetToolsImpl::MigratePackages(const TArray<FName>& PackageNamesToMigrate, const FString& DestinationPath, const struct FMigrationOptions& Options) const
 {
 	// Packages must be saved for the migration to work
-	const bool bPromptUserToSave = true;
+	const bool bPromptUserToSave = !FApp::IsUnattended() && Options.bPrompt;
 	const bool bSaveMapPackages = true;
 	const bool bSaveContentPackages = true;
-	if ( FEditorFileUtils::SaveDirtyPackages( bPromptUserToSave, bSaveMapPackages, bSaveContentPackages ) )
+	if (FEditorFileUtils::SaveDirtyPackages(bPromptUserToSave, bSaveMapPackages, bSaveContentPackages))
 	{
 		FAssetRegistryModule& AssetRegistryModule = FModuleManager::Get().LoadModuleChecked<FAssetRegistryModule>(TEXT("AssetRegistry"));
-		if ( AssetRegistryModule.Get().IsLoadingAssets() )
+		if (AssetRegistryModule.Get().IsLoadingAssets())
 		{
 			// Open a dialog asking the user to wait while assets are being discovered
 			SDiscoveringAssetsDialog::OpenDiscoveringAssetsDialog(
-				SDiscoveringAssetsDialog::FOnAssetsDiscovered::CreateUObject(this, &UAssetToolsImpl::PerformMigratePackages, PackageNamesToMigrate)
+				SDiscoveringAssetsDialog::FOnAssetsDiscovered::CreateUObject(this, &UAssetToolsImpl::PerformMigratePackages, PackageNamesToMigrate, DestinationPath, Options)
 			);
 		}
 		else
 		{
 			// Assets are already discovered, perform the migration now
-			PerformMigratePackages(PackageNamesToMigrate);
+			PerformMigratePackages(PackageNamesToMigrate, DestinationPath, Options);
 		}
 	}
 }
@@ -2951,7 +2958,7 @@ bool UAssetToolsImpl::CanCreateAsset(const FString& AssetName, const FString& Pa
 	return true;
 }
 
-void UAssetToolsImpl::PerformMigratePackages(TArray<FName> PackageNamesToMigrate) const
+void UAssetToolsImpl::PerformMigratePackages(TArray<FName> PackageNamesToMigrate, const FString DestinationPath, const FMigrationOptions Options) const
 {
 	// Form a full list of packages to move by including the dependencies of the supplied packages
 	TSet<FName> AllPackageNamesToMove;
@@ -3035,76 +3042,99 @@ void UAssetToolsImpl::PerformMigratePackages(TArray<FName> PackageNamesToMigrate
 	// Confirm that there is at least one package to move 
 	if (AllPackageNamesToMove.Num() == 0)
 	{
-		FMessageDialog::Open(EAppMsgType::Ok, LOCTEXT("MigratePackages_NoFilesToMove", "No files were found to move"));
+		if (!FApp::IsUnattended())
+		{
+			FMessageDialog::Open(EAppMsgType::Ok, LOCTEXT("MigratePackages_NoFilesToMove", "No files were found to move"));
+		}
 		return;
 	}
 
+	TSharedPtr<TArray<ReportPackageData>> ReportPackages = MakeShareable(new TArray<ReportPackageData>);
+	for (auto PackageIt = AllPackageNamesToMove.CreateConstIterator(); PackageIt; ++PackageIt)
+	{
+		bool bShouldMigratePackage = ShouldMigratePackage.Find(*PackageIt) != nullptr;
+		ReportPackages.Get()->Add({ (*PackageIt).ToString(), bShouldMigratePackage });
+	}
 	// Prompt the user displaying all assets that are going to be migrated
+	if(!FApp::IsUnattended() && Options.bPrompt)
 	{
 		const FText ReportMessage = LOCTEXT("MigratePackagesReportTitle", "The following assets will be migrated to another content folder.");
-		TSharedPtr<TArray<ReportPackageData>> ReportPackages = MakeShareable(new TArray<ReportPackageData>);
-		for ( auto PackageIt = AllPackageNamesToMove.CreateConstIterator(); PackageIt; ++PackageIt )
-		{
-			bool bShouldMigratePackage = ShouldMigratePackage.Find(*PackageIt) != nullptr;
-			ReportPackages.Get()->Add({ (*PackageIt).ToString(), bShouldMigratePackage });
-		}
-		SPackageReportDialog::FOnReportConfirmed OnReportConfirmed = SPackageReportDialog::FOnReportConfirmed::CreateUObject(this, &UAssetToolsImpl::MigratePackages_ReportConfirmed, ReportPackages);
+		SPackageReportDialog::FOnReportConfirmed OnReportConfirmed = SPackageReportDialog::FOnReportConfirmed::CreateUObject(this, &UAssetToolsImpl::MigratePackages_ReportConfirmed, ReportPackages, DestinationPath, Options);
 		SPackageReportDialog::OpenPackageReportDialog(ReportMessage, *ReportPackages.Get(), OnReportConfirmed);
 	}
+	else
+	{
+		UAssetToolsImpl::MigratePackages_ReportConfirmed(ReportPackages, DestinationPath, Options);
+	}
 }
 
-void UAssetToolsImpl::MigratePackages_ReportConfirmed(TSharedPtr<TArray<ReportPackageData>> PackageDataToMigrate) const
+void UAssetToolsImpl::MigratePackages_ReportConfirmed(TSharedPtr<TArray<ReportPackageData>> PackageDataToMigrate, const FString DestinationPath, const struct FMigrationOptions Options) const
 {
-	// Choose a destination folder
-	IDesktopPlatform* DesktopPlatform = FDesktopPlatformModule::Get();
 	FString DestinationFolder;
-	if ( ensure(DesktopPlatform) )
+	if (FApp::IsUnattended() || !DestinationPath.IsEmpty())
 	{
-		const void* ParentWindowWindowHandle = FSlateApplication::Get().FindBestParentWindowHandleForDialogs(nullptr);
+		if(DestinationPath.IsEmpty())
+		{
+			UE_LOG(LogAssetTools, Error, TEXT("Migration Destination path cannot be empty."));
+			return;
+		}
 
-		const FString Title = LOCTEXT("MigrateToFolderTitle", "Choose a destination Content folder").ToString();
-		bool bFolderAccepted = false;
-		while (!bFolderAccepted)
+		DestinationFolder = DestinationPath;
+		FPaths::NormalizeDirectoryName(DestinationFolder);
+		DestinationFolder += TEXT("/");
+	}
+	else
+	{
+		// Choose a destination folder
+		IDesktopPlatform* DesktopPlatform = FDesktopPlatformModule::Get();
+		if (ensure(DesktopPlatform))
 		{
-			const bool bFolderSelected = DesktopPlatform->OpenDirectoryDialog(
-				ParentWindowWindowHandle,
-				Title,
-				FEditorDirectories::Get().GetLastDirectory(ELastDirectory::GENERIC_EXPORT),
-				DestinationFolder
-				);
+			const void* ParentWindowWindowHandle = FSlateApplication::Get().FindBestParentWindowHandleForDialogs(nullptr);
 
-			if ( !bFolderSelected )
+			const FString Title = LOCTEXT("MigrateToFolderTitle", "Choose a destination Content folder").ToString();
+			bool bFolderAccepted = false;
+			while (!bFolderAccepted)
 			{
-				// User canceled, return
-				return;
-			}
+				const bool bFolderSelected = DesktopPlatform->OpenDirectoryDialog(
+					ParentWindowWindowHandle,
+					Title,
+					FEditorDirectories::Get().GetLastDirectory(ELastDirectory::GENERIC_EXPORT),
+					DestinationFolder
+				);
 
-			FEditorDirectories::Get().SetLastDirectory(ELastDirectory::GENERIC_EXPORT, DestinationFolder);
-			FPaths::NormalizeFilename(DestinationFolder);
-			if ( !DestinationFolder.EndsWith(TEXT("/")) )
-			{
-				DestinationFolder += TEXT("/");
-			}
+				if (!bFolderSelected)
+				{
+					// User canceled, return
+					return;
+				}
 
-			// Verify that it is a content folder
-			if ( DestinationFolder.EndsWith(TEXT("/Content/")) )
-			{
-				bFolderAccepted = true;
-			}
-			else
-			{
-				// The user chose a non-content folder. Confirm that this was their intention.
-				const FText Message = FText::Format(LOCTEXT("MigratePackages_NonContentFolder", "{0} does not appear to be a game Content folder. Migrated content will only work properly if placed in a Content folder. Would you like to place your content here anyway?"), FText::FromString(DestinationFolder));
-				EAppReturnType::Type Response = FMessageDialog::Open(EAppMsgType::YesNo, Message);
-				bFolderAccepted = (Response == EAppReturnType::Yes);
+				FEditorDirectories::Get().SetLastDirectory(ELastDirectory::GENERIC_EXPORT, DestinationFolder);
+				FPaths::NormalizeFilename(DestinationFolder);
+				if (!DestinationFolder.EndsWith(TEXT("/")))
+				{
+					DestinationFolder += TEXT("/");
+				}
+
+				// Verify that it is a content folder
+				if (DestinationFolder.EndsWith(TEXT("/Content/")))
+				{
+					bFolderAccepted = true;
+				}
+				else
+				{
+					// The user chose a non-content folder. Confirm that this was their intention.
+					const FText Message = FText::Format(LOCTEXT("MigratePackages_NonContentFolder", "{0} does not appear to be a game Content folder. Migrated content will only work properly if placed in a Content folder. Would you like to place your content here anyway?"), FText::FromString(DestinationFolder));
+					EAppReturnType::Type Response = FMessageDialog::Open(EAppMsgType::YesNo, Message);
+					bFolderAccepted = (Response == EAppReturnType::Yes);
+				}
 			}
 		}
-	}
-	else
-	{
-		// Not on a platform that supports desktop functionality
-		FMessageDialog::Open(EAppMsgType::Ok, LOCTEXT("NoDesktopPlatform", "Error: This platform does not support a file dialog."));
-		return;
+		else
+		{
+			// Not on a platform that supports desktop functionality
+			FMessageDialog::Open(EAppMsgType::Ok, LOCTEXT("NoDesktopPlatform", "Error: This platform does not support a file dialog."));
+			return;
+		}
 	}
 
 	// Build a list of packages to handle
@@ -3220,24 +3250,31 @@ void UAssetToolsImpl::MigratePackages_ReportConfirmed(TSharedPtr<TArray<ReportPa
 	// If some packages don't have a matching content root in the destination, prompt for desired behavior
 	if (!LostPackages.IsEmpty())
 	{
-		FString LostPackageRootsString;
-		for (const FName& PackageRoot : LostPackageRoots)
+		if (!FApp::IsUnattended() && Options.bPrompt && Options.OrphanFolder.IsEmpty())
 		{
-			LostPackageRootsString += FString(TEXT("\n\t")) + PackageRoot.ToString();
-		}
+			FString LostPackageRootsString;
+			for (const FName& PackageRoot : LostPackageRoots)
+			{
+				LostPackageRootsString += FString(TEXT("\n\t")) + PackageRoot.ToString();
+			}
 
-		// Prompt to consolidate to a migration folder
-		FText Prompt = FText::Format(LOCTEXT("MigratePackages_ConsolidateToTemp", "Some selected assets don't have a corresponding content root in the destination.{0}\n\nWould you like to save a copy of all selected assets into a folder with consolidated references? If you select No then assets in the above roots will not be migrated."), FText::FromString(LostPackageRootsString));
-		switch (FMessageDialog::Open(EAppMsgType::YesNoCancel, Prompt))
+			// Prompt to consolidate to a migration folder
+			FText Prompt = FText::Format(LOCTEXT("MigratePackages_ConsolidateToTemp", "Some selected assets don't have a corresponding content root in the destination.{0}\n\nWould you like to save a copy of all selected assets into a folder with consolidated references? If you select No then assets in the above roots will not be migrated."), FText::FromString(LostPackageRootsString));
+			switch (FMessageDialog::Open(EAppMsgType::YesNoCancel, Prompt))
+			{
+			case EAppReturnType::Yes:
+				// No op
+				break;
+			case EAppReturnType::No:
+				LostPackages.Reset();
+				break;
+			case EAppReturnType::Cancel:
+				return;
+			}
+		}
+		else if (Options.OrphanFolder.IsEmpty())
 		{
-		case EAppReturnType::Yes:
-			// No op
-			break;
-		case EAppReturnType::No:
 			LostPackages.Reset();
-			break;
-		case EAppReturnType::Cancel:
-			return;
 		}
 	}
 
@@ -3248,67 +3285,69 @@ void UAssetToolsImpl::MigratePackages_ReportConfirmed(TSharedPtr<TArray<ReportPa
 	if (!LostPackages.IsEmpty())
 	{
 		// Query the user for a folder to migrate assets into. This folder will exist temporary in this project, and is the destination for migration in the target project
-		FString FolderName;
-		bool bIsOkButtonEnabled = true;
-		TSharedRef<SEditableTextBox> EditableTextBox = SNew(SEditableTextBox)
-			.Text(FText::FromString(TEXT("Migrated")))
-			.OnVerifyTextChanged_Lambda([&bIsOkButtonEnabled](const FText& InNewText, FText& OutErrorMessage) -> bool
-				{
-					if (InNewText.ToString().Contains(TEXT("/")))
+		FString FolderName = Options.OrphanFolder.IsEmpty()? TEXT("Migrated") : Options.OrphanFolder;
+		if (!FApp::IsUnattended() && Options.bPrompt)
+		{
+			bool bIsOkButtonEnabled = true;
+			TSharedRef<SEditableTextBox> EditableTextBox = SNew(SEditableTextBox)
+				.Text(FText::FromString(FolderName))
+				.OnVerifyTextChanged_Lambda([&bIsOkButtonEnabled](const FText& InNewText, FText& OutErrorMessage) -> bool
 					{
-						OutErrorMessage = LOCTEXT("Migrated_CannotContainSlashes", "Cannot use a slash in a folder name.");
+						if (InNewText.ToString().Contains(TEXT("/")))
+						{
+							OutErrorMessage = LOCTEXT("Migrated_CannotContainSlashes", "Cannot use a slash in a folder name.");
 
-						// Disable Ok if the string is invalid
-						bIsOkButtonEnabled = false;
-						return false;
-					}
+							// Disable Ok if the string is invalid
+							bIsOkButtonEnabled = false;
+							return false;
+						}
+
+						// Enable Ok if the string is valid
+						bIsOkButtonEnabled = true;
+						return true;
+					})
+				.OnTextCommitted_Lambda([&FolderName](const FText& NewValue, ETextCommit::Type)
+					{
+						// Set the result if they modified the text
+						FolderName = NewValue.ToString();
+					});
 
-					// Enable Ok if the string is valid
-					bIsOkButtonEnabled = true;
-					return true;
-				})
-			.OnTextCommitted_Lambda([&FolderName](const FText& NewValue, ETextCommit::Type)
+			// Set the result if they just click Ok
+			SGenericDialogWidget::FArguments FolderDialogArguments;
+			FolderDialogArguments.OnOkPressed_Lambda([&EditableTextBox, &FolderName]()
 				{
-					// Set the result if they modified the text
-					FolderName = NewValue.ToString();
+					FolderName = EditableTextBox->GetText().ToString();
 				});
 
-		// Set the result if they just click Ok
-		SGenericDialogWidget::FArguments FolderDialogArguments;
-		FolderDialogArguments.OnOkPressed_Lambda([&EditableTextBox, &FolderName]()
-			{
-				FolderName = EditableTextBox->GetText().ToString();
-			});
-
-		// Present the Dialog
-		SGenericDialogWidget::OpenDialog(LOCTEXT("MigratePackages_FolderName", "Folder for Migrated Assets"),
-			SNew(SVerticalBox)
-			+ SVerticalBox::Slot()
+			// Present the Dialog
+			SGenericDialogWidget::OpenDialog(LOCTEXT("MigratePackages_FolderName", "Folder for Migrated Assets"),
+				SNew(SVerticalBox)
+				+ SVerticalBox::Slot()
 				.AutoHeight()
 				.Padding(5.0f)
-			[
-				SNew(STextBlock)
-				.Text(LOCTEXT("MigratePackages_SpecifyConsolidateFolder", "Please specify a new folder name to consolidate the assets into."))
-			]
+				[
+					SNew(STextBlock)
+					.Text(LOCTEXT("MigratePackages_SpecifyConsolidateFolder", "Please specify a new folder name to consolidate the assets into."))
+				]
 			+ SVerticalBox::Slot()
 				.Padding(5.0f)
-			[
-				SNew(SSpacer)
-			]
+				[
+					SNew(SSpacer)
+				]
 			+ SVerticalBox::Slot()
 				.AutoHeight()
 				.Padding(5.0f)
-			[
-				EditableTextBox
-			],
-			FolderDialogArguments, true);
+				[
+					EditableTextBox
+				],
+				FolderDialogArguments, true);
 
-		// Sanity the user input
-		if (FolderName.IsEmpty())
-		{
-			return;
+			// Sanity the user input
+			if (FolderName.IsEmpty())
+			{
+				return;
+			}
 		}
-
 		// Remove forbidden characters
 		FolderName = FolderName.Replace(TEXT("/"), TEXT(""));
 
@@ -3322,7 +3361,14 @@ void UAssetToolsImpl::MigratePackages_ReportConfirmed(TSharedPtr<TArray<ReportPa
 		if (IFileManager::Get().DirectoryExists(*SrcDiskFolderFilename))
 		{
 			const FText Message = FText::Format(LOCTEXT("MigratePackages_InvalidMigrateFolder", "{0} exists on disk in the source project, and cannot be used to consolidate assets."), FText::FromString(SrcDiskFolderFilename));
-			EAppReturnType::Type Response = FMessageDialog::Open(EAppMsgType::Ok, Message);
+			if (!FApp::IsUnattended() && Options.bPrompt)
+			{
+				EAppReturnType::Type Response = FMessageDialog::Open(EAppMsgType::Ok, Message);
+			}
+			else
+			{
+				UE_LOG(LogAssetTools, Error, TEXT("%s"), *Message.ToString());
+			}
 			return;
 		}
 
@@ -3453,6 +3499,20 @@ void UAssetToolsImpl::MigratePackages_ReportConfirmed(TSharedPtr<TArray<ReportPa
 	bool bUserCanceled = false;
 
 	EAppReturnType::Type LastResponse = EAppReturnType::Yes;
+	if (!Options.bPrompt)
+	{
+		switch (Options.AssetConflict) {
+		case EAssetMigrationConflict::Overwrite:
+			LastResponse = EAppReturnType::YesAll;
+			break;
+		case EAssetMigrationConflict::Skip:
+			LastResponse = EAppReturnType::NoAll;
+			break;
+		case EAssetMigrationConflict::Cancel:
+			LastResponse = EAppReturnType::Cancel;
+			break;
+		}
+	}
 	TArray<FString> SuccessfullyCopiedFiles;
 	TArray<FString> SuccessfullyCopiedPackages;
 	FString CopyErrors;
@@ -3509,7 +3569,7 @@ void UAssetToolsImpl::MigratePackages_ReportConfirmed(TSharedPtr<TArray<ReportPa
 					{
 						// The destination file already exists! Ask the user what to do.
 						EAppReturnType::Type Response;
-						if ( LastResponse == EAppReturnType::YesAll || LastResponse == EAppReturnType::NoAll )
+						if (FApp::IsUnattended() || !Options.bPrompt || LastResponse == EAppReturnType::YesAll || LastResponse == EAppReturnType::NoAll)
 						{
 							Response = LastResponse;
 						}
@@ -3517,15 +3577,16 @@ void UAssetToolsImpl::MigratePackages_ReportConfirmed(TSharedPtr<TArray<ReportPa
 						{
 							const FText Message = FText::Format( LOCTEXT("MigratePackages_AlreadyExists", "An asset already exists at location {0} would you like to overwrite it?"), FText::FromString(DestFilename) );
 							Response = FMessageDialog::Open( EAppMsgType::YesNoYesAllNoAllCancel, Message );
-							if ( Response == EAppReturnType::Cancel )
-							{
-								// The user chose to cancel mid-operation. Break out.
-								bUserCanceled = true;
-								break;
-							}
 							LastResponse = Response;
 						}
 
+						if (Response == EAppReturnType::Cancel)
+						{
+							// The user chose to cancel mid-operation. Break out.
+							bUserCanceled = true;
+							break;
+						}
+
 						const bool bWantOverwrite = Response == EAppReturnType::Yes || Response == EAppReturnType::YesAll;
 						if( !bWantOverwrite )
 						{
diff --git a/Engine/Source/Developer/AssetTools/Private/AssetTools.h b/Engine/Source/Developer/AssetTools/Private/AssetTools.h
--- a/Engine/Source/Developer/AssetTools/Private/AssetTools.h
+++ b/Engine/Source/Developer/AssetTools/Private/AssetTools.h
@@ -106,6 +106,7 @@ public:
 	virtual FString DumpAssetToTempFile(UObject* Asset) const override;
 	virtual bool CreateDiffProcess(const FString& DiffCommand, const FString& OldTextFilename, const FString& NewTextFilename, const FString& DiffArgs = FString("")) const override;
 	virtual void MigratePackages(const TArray<FName>& PackageNamesToMigrate) const override;
+	virtual void MigratePackages(const TArray<FName>& PackageNamesToMigrate, const FString& TargetPath, const struct FMigrationOptions& Options = FMigrationOptions()) const override;
 	virtual void BeginAdvancedCopyPackages(const TArray<FName>& InputNamesToCopy, const FString& TargetPath) const override;
 	virtual void FixupReferencers(const TArray<UObjectRedirector*>& Objects, bool bCheckoutDialogPrompt = true, ERedirectFixupMode FixupMode = ERedirectFixupMode::DeleteFixedUpRedirectors) const override;
 	virtual bool IsFixupReferencersInProgress() const override;
@@ -158,7 +159,7 @@ private:
 	bool CanCreateAsset(const FString& AssetName, const FString& PackageName, const FText& OperationText) const;
 
 	/** Begins the package migration, after assets have been discovered */
-	void PerformMigratePackages(TArray<FName> PackageNamesToMigrate) const;
+	void PerformMigratePackages(TArray<FName> PackageNamesToMigrate, const FString DestinationPath, const FMigrationOptions Options) const;
 
 	TArray<FName> ExpandAssetsAndFoldersToJustAssets(TArray<FName> SelectedAssetAndFolderNames) const;
 
@@ -166,7 +167,7 @@ private:
 	void PerformAdvancedCopyPackages(TArray<FName> SelectedPackageNames, FString TargetPath) const;
 
 	/** Copies files after the final list was confirmed */
-	void MigratePackages_ReportConfirmed(TSharedPtr<TArray<ReportPackageData>> PackageDataToMigrate) const;
+	void MigratePackages_ReportConfirmed(TSharedPtr<TArray<ReportPackageData>> PackageDataToMigrate, const FString DestinationPath, const FMigrationOptions Options) const;
 
 	/** Copies files after the final list was confirmed */
 	void AdvancedCopyPackages_ReportConfirmed(FAdvancedCopyParams CopyParam, TArray<TMap<FString, FString>> DestinationMap) const;
diff --git a/Engine/Source/Developer/AssetTools/Public/IAssetTools.h b/Engine/Source/Developer/AssetTools/Public/IAssetTools.h
--- a/Engine/Source/Developer/AssetTools/Public/IAssetTools.h
+++ b/Engine/Source/Developer/AssetTools/Public/IAssetTools.h
@@ -177,6 +177,41 @@ private:
 	FString DropLocationForAdvancedCopy;
 };
 
+
+UENUM()
+enum class EAssetMigrationConflict : uint8
+{
+	// Skip Assets
+	Skip,
+	// Overwrite Assets
+	Overwrite,
+	// Cancel the whole Migration
+	Cancel
+};
+
+USTRUCT(BlueprintType)
+struct FMigrationOptions
+{
+	GENERATED_BODY()
+
+	/** Prompt user for confirmation (always false through scripting) */
+	bool bPrompt;
+
+	/** What to do when Assets are conflicting on the destination */
+	UPROPERTY(BlueprintReadWrite, Category = MigrationOptions)
+	EAssetMigrationConflict AssetConflict;
+
+	/** Destination for assets that don't have a corresponding content folder. If left empty those assets are not migrated. */
+	UPROPERTY(BlueprintReadWrite, Category = MigrationOptions)
+	FString OrphanFolder;
+
+	FMigrationOptions()
+		: bPrompt(false)
+		, AssetConflict(EAssetMigrationConflict::Skip)
+	{}
+
+};
+
 // An array of maps each storing pairs of original object -> duplicated object
 using FDuplicatedObjects = TArray<TMap<TSoftObjectPtr<UObject>, TSoftObjectPtr<UObject>>>;
 
@@ -418,6 +453,10 @@ public:
 	/* Migrate packages to another game content folder */
 	virtual void MigratePackages(const TArray<FName>& PackageNamesToMigrate) const = 0;
 
+	/* Migrate packages and dependencies to another folder */
+	UFUNCTION(BlueprintCallable, Category = "Editor Scripting | Asset Tools")
+	virtual void MigratePackages(const TArray<FName>& PackageNamesToMigrate, const FString& DestinationPath, const struct FMigrationOptions& Options = FMigrationOptions()) const = 0;
+
 	/* Copy packages and dependencies to another folder */
 	virtual void BeginAdvancedCopyPackages(const TArray<FName>& InputNamesToCopy, const FString& TargetPath) const = 0;
 
