Convert usage of GIsEditorLoadingPackage to TGuardValue.

diff --git a/Engine/Source/Developer/AssetTools/Private/AssetViewUtils.cpp b/Engine/Source/Developer/AssetTools/Private/AssetViewUtils.cpp
--- a/Engine/Source/Developer/AssetTools/Private/AssetViewUtils.cpp
+++ b/Engine/Source/Developer/AssetTools/Private/AssetViewUtils.cpp
@@ -166,40 +166,41 @@ bool AssetViewUtils::LoadAssetsIfNeeded(const TArray<FString>& ObjectPaths, TArr
 			SlowTask.MakeDialog();
 		}
 
-		GIsEditorLoadingPackage = true;
-
-		// We usually don't want to follow redirects when loading objects for the Content Browser.  It would
-		// allow a user to interact with a ghost/unverified asset as if it were still alive.
-		// This can be overridden by providing bLoadRedirects = true as a parameter.
-		const ELoadFlags LoadFlags = bLoadRedirects ? LOAD_None : LOAD_NoRedirects;
-
 		bool bSomeObjectsFailedToLoad = false;
-		for (int32 PathIdx = 0; PathIdx < UnloadedObjectPaths.Num(); ++PathIdx)
 		{
-			const FString& ObjectPath = UnloadedObjectPaths[PathIdx];
-			SlowTask.EnterProgressFrame(1, FText::Format(LOCTEXT("LoadingObjectf", "Loading {0}..."), FText::FromString(ObjectPath)));
+			TGuardValue<bool> IsEditorLoadingPackageGuard(GIsEditorLoadingPackage, true);
 
-			// Load up the object
-			FLinkerInstancingContext InstancingContext(bLoadAllExternalObjects ? TSet<FName>{ ULevel::LoadAllExternalObjectsTag } : TSet<FName>());
-			UObject* LoadedObject = LoadObject<UObject>(NULL, *ObjectPath, NULL, LoadFlags, NULL, &InstancingContext);
-			if ( LoadedObject )
-			{
-				LoadedObjects.Add(LoadedObject);
-			}
-			else
-			{
-				bSomeObjectsFailedToLoad = true;
-			}
+			// We usually don't want to follow redirects when loading objects for the Content Browser.  It would
+			// allow a user to interact with a ghost/unverified asset as if it were still alive.
+			// This can be overridden by providing bLoadRedirects = true as a parameter.
+			const ELoadFlags LoadFlags = bLoadRedirects ? LOAD_None : LOAD_NoRedirects;
 
-			if (GWarn->ReceivedUserCancel())
+			for (int32 PathIdx = 0; PathIdx < UnloadedObjectPaths.Num(); ++PathIdx)
 			{
-				// If the user has canceled stop loading the remaining objects. We don't add the remaining objects to the failed string,
-				// this would only result in launching another dialog when by their actions the user clearly knows not all of the 
-				// assets will have been loaded.
-				break;
+				const FString& ObjectPath = UnloadedObjectPaths[PathIdx];
+				SlowTask.EnterProgressFrame(1, FText::Format(LOCTEXT("LoadingObjectf", "Loading {0}..."), FText::FromString(ObjectPath)));
+
+				// Load up the object
+				FLinkerInstancingContext InstancingContext(bLoadAllExternalObjects ? TSet<FName>{ ULevel::LoadAllExternalObjectsTag } : TSet<FName>());
+				UObject* LoadedObject = LoadObject<UObject>(NULL, *ObjectPath, NULL, LoadFlags, NULL, &InstancingContext);
+				if ( LoadedObject )
+				{
+					LoadedObjects.Add(LoadedObject);
+				}
+				else
+				{
+					bSomeObjectsFailedToLoad = true;
+				}
+
+				if (GWarn->ReceivedUserCancel())
+				{
+					// If the user has canceled stop loading the remaining objects. We don't add the remaining objects to the failed string,
+					// this would only result in launching another dialog when by their actions the user clearly knows not all of the 
+					// assets will have been loaded.
+					break;
+				}
 			}
 		}
-		GIsEditorLoadingPackage = false;
 
 		if ( bSomeObjectsFailedToLoad )
 		{
diff --git a/Engine/Source/Editor/Kismet/Private/FindInBlueprints.cpp b/Engine/Source/Editor/Kismet/Private/FindInBlueprints.cpp
--- a/Engine/Source/Editor/Kismet/Private/FindInBlueprints.cpp
+++ b/Engine/Source/Editor/Kismet/Private/FindInBlueprints.cpp
@@ -244,9 +244,11 @@ UBlueprint* FFindInBlueprintsResult::GetParentBlueprint() const
 	}
 	else
 	{
-		GIsEditorLoadingPackage = true;
-		UObject* Object = LoadObject<UObject>(NULL, *DisplayText.ToString(), NULL, 0, NULL);
-		GIsEditorLoadingPackage = false;
+		UObject* Object;
+		{
+			TGuardValue<bool> IsEditorLoadingPackageGuard(GIsEditorLoadingPackage, true);
+			Object = LoadObject<UObject>(NULL, *DisplayText.ToString(), NULL, 0, NULL);
+		}
 
 		if(UBlueprint* BlueprintObj = Cast<UBlueprint>(Object))
 		{
diff --git a/Engine/Source/Editor/UnrealEd/Private/EditorServer.cpp b/Engine/Source/Editor/UnrealEd/Private/EditorServer.cpp
--- a/Engine/Source/Editor/UnrealEd/Private/EditorServer.cpp
+++ b/Engine/Source/Editor/UnrealEd/Private/EditorServer.cpp
@@ -2481,7 +2481,7 @@ bool UEditorEngine::Map_Load(const TCHAR* Str, FOutputDevice& Ar)
 
 #define LOCTEXT_NAMESPACE "EditorEngine"
 	// We are beginning a map load
-	GIsEditorLoadingPackage = true;
+	TGuardValue<bool> IsEditorLoadingPackageGuard(GIsEditorLoadingPackage, true);
 
 	FWorldContext &Context = GetEditorWorldContext();
 	check(Context.World() == GWorld);
@@ -2510,7 +2510,6 @@ bool UEditorEngine::Map_Load(const TCHAR* Str, FOutputDevice& Ar)
 					FFormatNamedArguments Arguments;
 					Arguments.Add(TEXT("Reason"), NotMapReason);
 					FMessageDialog::Open(EAppMsgType::Ok, FText::Format(LOCTEXT("MapLoadFailed", "Failed to load map!\n{Reason}"), Arguments));
-					GIsEditorLoadingPackage = false;
 					return false;
 				}
 
@@ -2730,7 +2729,6 @@ bool UEditorEngine::Map_Load(const TCHAR* Str, FOutputDevice& Ar)
 				if (WorldPackage == nullptr)
 				{
 					FMessageDialog::Open( EAppMsgType::Ok, NSLOCTEXT("UnrealEd", "MapPackageLoadFailed", "Failed to open map file. This is most likely because the map was saved with a newer version of the engine."));
-					GIsEditorLoadingPackage = false;
 					return false;
 				}
 
@@ -2743,7 +2741,6 @@ bool UEditorEngine::Map_Load(const TCHAR* Str, FOutputDevice& Ar)
 				{
 					FText Message = FText::Format(NSLOCTEXT("UnrealEd", "MapPackageFindWorldFailed", "Failed to find the world in already loaded world package {0}! See log for more details."), FText::FromString(WorldPackage->GetPathName()));
 					FMessageDialog::Open(EAppMsgType::Ok, Message);
-					GIsEditorLoadingPackage = false;
 
 					FReferenceChainSearch RefChainSearch(WorldPackage, EReferenceChainSearchMode::Shortest | EReferenceChainSearchMode::PrintResults);
 					UE_LOG(LogEditorServer, Warning, TEXT("Failed to find the world in already loaded world package %s! Referenced by:") LINE_TERMINATOR TEXT("%s"), *WorldPackage->GetPathName(), *RefChainSearch.GetRootPath());
@@ -2943,7 +2940,6 @@ bool UEditorEngine::Map_Load(const TCHAR* Str, FOutputDevice& Ar)
 	}
 
 	// Done loading a map
-	GIsEditorLoadingPackage = false;
 	return true;
 #undef LOCTEXT_NAMESPACE
 }
