Fix incorrect asset menu actions appearing for animation blueprints

Also fixes the asset context menu section header for anim BPs to not read as "animation layer interface" (this is just the re-ordering of asset type action registration)

diff --git a/Engine/Source/Developer/AssetTools/Private/AssetTypeActions/AssetTypeActions_AnimBlueprint.cpp b/Engine/Source/Developer/AssetTools/Private/AssetTypeActions/AssetTypeActions_AnimBlueprint.cpp
--- a/Engine/Source/Developer/AssetTools/Private/AssetTypeActions/AssetTypeActions_AnimBlueprint.cpp
+++ b/Engine/Source/Developer/AssetTools/Private/AssetTypeActions/AssetTypeActions_AnimBlueprint.cpp
@@ -30,14 +30,18 @@
 
 void FAssetTypeActions_AnimBlueprint::GetActions(const TArray<UObject*>& InObjects, FToolMenuSection& Section)
 {
-	FAssetTypeActions_Blueprint::GetActions(InObjects, Section);
-
 	TArray<TWeakObjectPtr<UAnimBlueprint>> AnimBlueprints = GetTypedWeakObjectPtrs<UAnimBlueprint>(InObjects);
 
 	if (AnimBlueprints.Num() == 1 && CanCreateNewDerivedBlueprint())
 	{
 		UAnimBlueprint* AnimBlueprint = AnimBlueprints[0].Get();
 
+		if(AnimBlueprint && AnimBlueprint->BlueprintType != BPTYPE_Interface)
+		{
+			// Call base class for the regular 'create child blueprint' option
+			FAssetTypeActions_Blueprint::GetActions(InObjects, Section);
+		}
+		
 		// Accept (non-interface) template anim BPs or anim BPs with compatible skeletons
 		if(AnimBlueprint && AnimBlueprint->BlueprintType != BPTYPE_Interface && ((AnimBlueprint->TargetSkeleton == nullptr && AnimBlueprint->bIsTemplate) || (AnimBlueprint->TargetSkeleton != nullptr && AnimBlueprint->TargetSkeleton->GetCompatibleSkeletons().Num() > 0)))
 		{
@@ -129,7 +133,7 @@ void FAssetTypeActions_AnimBlueprint::GetActions(const TArray<UObject*>& InObjec
 		}
 	}
 
-	if(AreAnyNonTemplateAnimBlueprintsSelected(AnimBlueprints))
+	if(AreOnlyNonTemplateAnimBlueprintsSelected(AnimBlueprints) && AreOnlyNonInterfaceAnimBlueprintsSelected(AnimBlueprints))
 	{
 		Section.AddMenuEntry(
 			"AnimBlueprint_FindSkeleton",
@@ -258,17 +262,30 @@ void FAssetTypeActions_AnimBlueprint::PerformAssetDiff(UObject* Asset1, UObject*
 	SBlueprintDiff::CreateDiffWindow(WindowTitle, OldBlueprint, NewBlueprint, OldRevision, NewRevision);
 }
 
-bool FAssetTypeActions_AnimBlueprint::AreAnyNonTemplateAnimBlueprintsSelected(TArray<TWeakObjectPtr<UAnimBlueprint>> Objects) const
+bool FAssetTypeActions_AnimBlueprint::AreOnlyNonTemplateAnimBlueprintsSelected(TArray<TWeakObjectPtr<UAnimBlueprint>> Objects) const
+{
+	for(TWeakObjectPtr<UAnimBlueprint> WeakAnimBlueprint : Objects)
+	{
+		if(WeakAnimBlueprint->bIsTemplate)
+		{
+			return false; 
+		}
+	}
+
+	return true;
+}
+
+bool FAssetTypeActions_AnimBlueprint::AreOnlyNonInterfaceAnimBlueprintsSelected(TArray<TWeakObjectPtr<UAnimBlueprint>> Objects) const
 {
 	for(TWeakObjectPtr<UAnimBlueprint> WeakAnimBlueprint : Objects)
 	{
-		if(!WeakAnimBlueprint->bIsTemplate)
+		if(WeakAnimBlueprint->BlueprintType == BPTYPE_Interface)
 		{
-			return true; 
+			return false; 
 		}
 	}
 
-	return false;
+	return true;
 }
 
 void FAssetTypeActions_AnimBlueprint::ExecuteFindSkeleton(TArray<TWeakObjectPtr<UAnimBlueprint>> Objects)
diff --git a/Engine/Source/Developer/AssetTools/Private/AssetTypeActions/AssetTypeActions_AnimBlueprint.h b/Engine/Source/Developer/AssetTools/Private/AssetTypeActions/AssetTypeActions_AnimBlueprint.h
--- a/Engine/Source/Developer/AssetTools/Private/AssetTypeActions/AssetTypeActions_AnimBlueprint.h
+++ b/Engine/Source/Developer/AssetTools/Private/AssetTypeActions/AssetTypeActions_AnimBlueprint.h
@@ -37,8 +37,11 @@ private:
 	void ExecuteAssignSkeleton(TArray<TWeakObjectPtr<UAnimBlueprint>> Objects);
 	
 	/** Certain options are not available for template anim BPs */
-	bool AreAnyNonTemplateAnimBlueprintsSelected(TArray<TWeakObjectPtr<UAnimBlueprint>> Objects) const;
+	bool AreOnlyNonTemplateAnimBlueprintsSelected(TArray<TWeakObjectPtr<UAnimBlueprint>> Objects) const;
 
+	/** Certain options are not available for layer interfaces */
+	bool AreOnlyNonInterfaceAnimBlueprintsSelected(TArray<TWeakObjectPtr<UAnimBlueprint>> Objects) const;
+	
 	/** When skeleton asset is missing, allow replacing skeleton asset */ 
 	bool ReplaceMissingSkeleton(TArray<UObject*> InAnimBlueprints) const;
 };
diff --git a/Engine/Source/Developer/AssetTools/Private/AssetTypeActions/AssetTypeActions_AnimBlueprintInterface.h b/Engine/Source/Developer/AssetTools/Private/AssetTypeActions/AssetTypeActions_AnimBlueprintInterface.h
--- a/Engine/Source/Developer/AssetTools/Private/AssetTypeActions/AssetTypeActions_AnimBlueprintInterface.h
+++ b/Engine/Source/Developer/AssetTools/Private/AssetTypeActions/AssetTypeActions_AnimBlueprintInterface.h
@@ -12,5 +12,4 @@ public:
 	virtual FText GetName() const override { return NSLOCTEXT("AssetTypeActions", "AssetTypeActions_AnimLayerInterface", "Animation Layer Interface"); }
 	virtual void BuildBackendFilter(FARFilter& InFilter) override;
 	virtual FName GetFilterName() const override;
-	virtual void GetActions(const TArray<UObject*>& InObjects, struct FToolMenuSection& Section) override {}
 };
