Fixed bad merge that led to SubUVAnimation asset actions being registered twice

#rnx

diff --git a/Engine/Source/Developer/AssetTools/Private/AssetTools.cpp b/Engine/Source/Developer/AssetTools/Private/AssetTools.cpp
--- a/Engine/Source/Developer/AssetTools/Private/AssetTools.cpp
+++ b/Engine/Source/Developer/AssetTools/Private/AssetTools.cpp
@@ -326,7 +326,6 @@ UAssetToolsImpl::UAssetToolsImpl(const FObjectInitializer& ObjectInitializer)
 	RegisterAssetTypeActions(MakeShareable(new FAssetTypeActions_SlateBrush));
 	RegisterAssetTypeActions(MakeShareable(new FAssetTypeActions_SlateWidgetStyle));
 	RegisterAssetTypeActions(MakeShareable(new FAssetTypeActions_StaticMesh));
-	RegisterAssetTypeActions(MakeShareable(new FAssetTypeActions_SubUVAnimation));
 	RegisterAssetTypeActions(MakeShareable(new FAssetTypeActions_Texture));
 	RegisterAssetTypeActions(MakeShareable(new FAssetTypeActions_Texture2D));
 	RegisterAssetTypeActions(MakeShareable(new FAssetTypeActions_TextureCube));
