Editor - Renable HasActions, since I'm working on a different replacement, I might as well clear build errors.

diff --git a/Engine/Source/Developer/AssetTools/Public/IAssetTypeActions.h b/Engine/Source/Developer/AssetTools/Public/IAssetTypeActions.h
--- a/Engine/Source/Developer/AssetTools/Public/IAssetTypeActions.h
+++ b/Engine/Source/Developer/AssetTools/Public/IAssetTypeActions.h
@@ -74,8 +74,8 @@ public:
 	virtual FColor GetTypeColor() const = 0;
 	
 	/** Returns true if this class can supply actions for InObjects. */
-   	UE_DEPRECATED(5.2, "HasActions is no longer used all IAssetTypeActions are assumed to need to have them and GetActions is just called instead.  GetActions is always called unless DEPRECATE_ASSET_TYPE_ACTIONS_CALLING_GETACTIONS is true.")
-   	virtual bool HasActions( const TArray<UObject*>& InObjects ) const final { return true; }
+//   	UE_DEPRECATED(5.2, "HasActions is no longer used all IAssetTypeActions are assumed to need to have them and GetActions is just called instead.  GetActions is always called unless DEPRECATE_ASSET_TYPE_ACTIONS_CALLING_GETACTIONS is true.")
+   	virtual bool HasActions( const TArray<UObject*>& InObjects ) const { return true; }
 
 #if DEPRECATE_ASSET_TYPE_ACTIONS_NEEDING_LOADED_OBJECTS
 	/** Generates a menubuilder for the specified objects. */
