Move Force Feedback effects and force feedback attenuation asset types to the "Input" category of the content browser categories instead of "Miscellaneous". Now that enhanced input already made an Input category, these asset types have a better place to live.

diff --git a/Engine/Source/Developer/AssetTools/Private/AssetTools.cpp b/Engine/Source/Developer/AssetTools/Private/AssetTools.cpp
--- a/Engine/Source/Developer/AssetTools/Private/AssetTools.cpp
+++ b/Engine/Source/Developer/AssetTools/Private/AssetTools.cpp
@@ -251,6 +251,7 @@ UAssetToolsImpl::UAssetToolsImpl(const FObjectInitializer& ObjectInitializer)
 	AllocatedCategoryBits.Add(TEXT("_BuiltIn_9"), FAdvancedAssetCategory(EAssetTypeCategories::Textures, LOCTEXT("TextureAssetCategory", "Textures")));
 
 	EAssetTypeCategories::Type FoliageCategoryBit = RegisterAdvancedAssetCategory(FName(TEXT("Foliage")), LOCTEXT("FoliageAssetCategory", "Foliage"));
+	EAssetTypeCategories::Type InputCategoryBit = RegisterAdvancedAssetCategory(FName(TEXT("Input")), LOCTEXT("InputAssetsCategory", "Input"));
 
 	// Register the built-in asset type actions
 	RegisterAssetTypeActions(MakeShareable(new FAssetTypeActions_Actor));
@@ -290,8 +291,8 @@ UAssetToolsImpl::UAssetToolsImpl(const FObjectInitializer& ObjectInitializer)
 	RegisterAssetTypeActions(MakeShareable(new FAssetTypeActions_SceneImportData));
 	RegisterAssetTypeActions(MakeShareable(new FAssetTypeActions_Font));
 	RegisterAssetTypeActions(MakeShareable(new FAssetTypeActions_FontFace));
-	RegisterAssetTypeActions(MakeShareable(new FAssetTypeActions_ForceFeedbackAttenuation));
-	RegisterAssetTypeActions(MakeShareable(new FAssetTypeActions_ForceFeedbackEffect));
+	RegisterAssetTypeActions(MakeShareable(new FAssetTypeActions_ForceFeedbackAttenuation(InputCategoryBit)));
+	RegisterAssetTypeActions(MakeShareable(new FAssetTypeActions_ForceFeedbackEffect(InputCategoryBit)));
 	RegisterAssetTypeActions(MakeShareable(new FAssetTypeActions_HLODProxy));
 	RegisterAssetTypeActions(MakeShareable(new FAssetTypeActions_SubsurfaceProfile));
 	RegisterAssetTypeActions(MakeShareable(new FAssetTypeActions_InstancedFoliageSettings(FoliageCategoryBit)));
diff --git a/Engine/Source/Developer/AssetTools/Private/AssetTypeActions/AssetTypeActions_ForceFeedbackAttenuation.h b/Engine/Source/Developer/AssetTools/Private/AssetTypeActions/AssetTypeActions_ForceFeedbackAttenuation.h
--- a/Engine/Source/Developer/AssetTools/Private/AssetTypeActions/AssetTypeActions_ForceFeedbackAttenuation.h
+++ b/Engine/Source/Developer/AssetTools/Private/AssetTypeActions/AssetTypeActions_ForceFeedbackAttenuation.h
@@ -9,9 +9,17 @@
 class FAssetTypeActions_ForceFeedbackAttenuation : public FAssetTypeActions_Base
 {
 public:
+
+	FAssetTypeActions_ForceFeedbackAttenuation(EAssetTypeCategories::Type InAssetCategoryBit)
+		: AssetCategoryBit(InAssetCategoryBit)
+	{ }
+
 	// IAssetTypeActions Implementation
 	virtual FText GetName() const override { return NSLOCTEXT("AssetTypeActions", "AssetTypeActions_ForceFeedbackAttenuation", "Force Feedback Attenuation"); }
 	virtual FColor GetTypeColor() const override { return FColor(175, 0, 0); }
 	virtual UClass* GetSupportedClass() const override { return UForceFeedbackAttenuation::StaticClass(); }
-	virtual uint32 GetCategories() override { return EAssetTypeCategories::Misc; }
+	virtual uint32 GetCategories() override { return AssetCategoryBit; }
+
+private:
+	EAssetTypeCategories::Type AssetCategoryBit;
 };
\ No newline at end of file
diff --git a/Engine/Source/Developer/AssetTools/Private/AssetTypeActions/AssetTypeActions_ForceFeedbackEffect.h b/Engine/Source/Developer/AssetTools/Private/AssetTypeActions/AssetTypeActions_ForceFeedbackEffect.h
--- a/Engine/Source/Developer/AssetTools/Private/AssetTypeActions/AssetTypeActions_ForceFeedbackEffect.h
+++ b/Engine/Source/Developer/AssetTools/Private/AssetTypeActions/AssetTypeActions_ForceFeedbackEffect.h
@@ -29,11 +29,15 @@ struct FPreviewForceFeedbackEffect : public FActiveForceFeedbackEffect, public F
 class FAssetTypeActions_ForceFeedbackEffect : public FAssetTypeActions_Base
 {
 public:
+	FAssetTypeActions_ForceFeedbackEffect(EAssetTypeCategories::Type InAssetCategoryBit)
+		: AssetCategoryBit(InAssetCategoryBit)
+	{ }
+
 	// IAssetTypeActions Implementation
 	virtual FText GetName() const override { return NSLOCTEXT("AssetTypeActions", "AssetTypeActions_ForceFeedbackEffect", "Force Feedback Effect"); }
 	virtual FColor GetTypeColor() const override { return FColor(175, 0, 0); }
 	virtual UClass* GetSupportedClass() const override;
-	virtual uint32 GetCategories() override { return EAssetTypeCategories::Misc; }
+	virtual uint32 GetCategories() override { return AssetCategoryBit; }
 	virtual void GetActions(const TArray<UObject*>& InObjects, struct FToolMenuSection& Section) override;
 	virtual bool AssetsActivatedOverride( const TArray<UObject*>& InObjects, EAssetTypeActivationMethod::Type ActivationType ) override;
 	virtual TSharedPtr<SWidget> GetThumbnailOverlay(const FAssetData& AssetData) const override;
@@ -64,4 +68,6 @@ public:
 
 private:
 	FPreviewForceFeedbackEffect PreviewForceFeedbackEffect;
+
+	EAssetTypeCategories::Type AssetCategoryBit;
 };
