Add a notification severity filter param on AdvancedCopyPackages

diff --git a/Engine/Source/Developer/AssetTools/Private/AssetTools.cpp b/Engine/Source/Developer/AssetTools/Private/AssetTools.cpp
--- a/Engine/Source/Developer/AssetTools/Private/AssetTools.cpp
+++ b/Engine/Source/Developer/AssetTools/Private/AssetTools.cpp
@@ -958,7 +958,8 @@ bool UAssetToolsImpl::AdvancedCopyPackages(
 	const TMap<FString, FString>& SourceAndDestPackages,
 	const bool bForceAutosave,
 	const bool bCopyOverAllDestinationOverlaps,
-	FDuplicatedObjects* OutDuplicatedObjects) const
+	FDuplicatedObjects* OutDuplicatedObjects,
+	EMessageSeverity::Type NotificationSeverityFilter) const
 {
 	if (ValidateFlattenedAdvancedCopyDestinations(SourceAndDestPackages))
 	{
@@ -1117,21 +1118,16 @@ bool UAssetToolsImpl::AdvancedCopyPackages(
 		EMessageSeverity::Type Severity = EMessageSeverity::Info;
 		if (SourceControlErrors.Len() > 0)
 		{
-			FString ErrorMessage;
 			Severity = EMessageSeverity::Error;
-			if (SourceControlErrors.Len() > 0)
-			{
-				AdvancedCopyLog.NewPage(LOCTEXT("AdvancedCopyPackages_SourceControlErrorsListPage", "Source Control Errors"));
-				AdvancedCopyLog.Error(FText::FromString(*SourceControlErrors));
-				ErrorMessage += LINE_TERMINATOR;
-				ErrorMessage += LOCTEXT("AdvancedCopyPackages_SourceControlErrorsList", "Some files reported source control errors.").ToString();
-			}
+
+			AdvancedCopyLog.NewPage(LOCTEXT("AdvancedCopyPackages_SourceControlErrorsListPage", "Source Control Errors"));
+			AdvancedCopyLog.Error(FText::FromString(*SourceControlErrors));
+
+			FString ErrorMessage = LOCTEXT("AdvancedCopyPackages_SourceControlErrorsList", "Some files reported source control errors.").ToString();
+			
 			if (SuccessfullyCopiedSourcePackages.Num() > 0)
 			{
 				AdvancedCopyLog.NewPage(LOCTEXT("AdvancedCopyPackages_CopyErrorsSuccesslistPage", "Copied Successfully"));
-				AdvancedCopyLog.Info(FText::FromString(*SourceControlErrors));
-				ErrorMessage += LINE_TERMINATOR;
-				ErrorMessage += LOCTEXT("AdvancedCopyPackages_CopyErrorsSuccesslist", "Some files were copied successfully.").ToString();
 				for (auto FileIt = SuccessfullyCopiedSourcePackages.CreateConstIterator(); FileIt; ++FileIt)
 				{
 					if (!FileIt->IsNone())
@@ -1139,6 +1135,9 @@ bool UAssetToolsImpl::AdvancedCopyPackages(
 						AdvancedCopyLog.Info(FText::FromName(*FileIt));
 					}
 				}
+
+				ErrorMessage += LINE_TERMINATOR;
+				ErrorMessage += LOCTEXT("AdvancedCopyPackages_CopyErrorsSuccesslist", "Some files were copied successfully.").ToString();
 			}
 			LogMessage = FText::FromString(ErrorMessage);
 		}
@@ -1153,7 +1152,8 @@ bool UAssetToolsImpl::AdvancedCopyPackages(
 				}
 			}
 		}
-		AdvancedCopyLog.Notify(LogMessage, Severity, true);
+		// @note Using the bForce param because the InSeverityFilter param is only checked against logs in the last page
+		AdvancedCopyLog.Notify(LogMessage, /*InSeverityFilter=*/EMessageSeverity::CriticalError, /*bForce=*/(Severity <= NotificationSeverityFilter));
 		return true;
 	}
 	return false;
@@ -1164,7 +1164,7 @@ bool UAssetToolsImpl::AdvancedCopyPackages(const FAdvancedCopyParams& CopyParams
 	TMap<FString, FString> FlattenedDestinationMap;
 	if (FlattenAdvancedCopyDestinations(PackagesAndDestinations, FlattenedDestinationMap))
 	{
-		return AdvancedCopyPackages(FlattenedDestinationMap, CopyParams.bShouldForceSave, CopyParams.bCopyOverAllDestinationOverlaps, nullptr);
+		return AdvancedCopyPackages(FlattenedDestinationMap, CopyParams.bShouldForceSave, CopyParams.bCopyOverAllDestinationOverlaps, /*OutDuplicatedObjects=*/nullptr, EMessageSeverity::Info);
 	}
 	return false;
 }
diff --git a/Engine/Source/Developer/AssetTools/Private/AssetTools.h b/Engine/Source/Developer/AssetTools/Private/AssetTools.h
--- a/Engine/Source/Developer/AssetTools/Private/AssetTools.h
+++ b/Engine/Source/Developer/AssetTools/Private/AssetTools.h
@@ -113,7 +113,7 @@ public:
 	virtual FAssetPostRenameEvent& OnAssetPostRename() override { return AssetRenameManager->OnAssetPostRenameEvent(); }
 	virtual void ExpandDirectories(const TArray<FString>& Files, const FString& DestinationPath, TArray<TPair<FString, FString>>& FilesAndDestinations) const override;
 	virtual bool AdvancedCopyPackages(const FAdvancedCopyParams& CopyParams, const TArray<TMap<FString, FString>> PackagesAndDestinations) const override;
-	virtual bool AdvancedCopyPackages(const TMap<FString, FString>& SourceAndDestPackages, const bool bForceAutosave, const bool bCopyOverAllDestinationOverlaps, FDuplicatedObjects* OutDuplicatedObjects) const override;
+	virtual bool AdvancedCopyPackages(const TMap<FString, FString>& SourceAndDestPackages, const bool bForceAutosave, const bool bCopyOverAllDestinationOverlaps, FDuplicatedObjects* OutDuplicatedObjects, EMessageSeverity::Type NotificationSeverityFilter) const override;
 	virtual void GenerateAdvancedCopyDestinations(FAdvancedCopyParams& InParams, const TArray<FName>& InPackageNamesToCopy, const class UAdvancedCopyCustomization* CopyCustomization, TMap<FString, FString>& OutPackagesAndDestinations) const override;
 	virtual bool FlattenAdvancedCopyDestinations(const TArray<TMap<FString, FString>> PackagesAndDestinations, TMap<FString, FString>& FlattenedPackagesAndDestinations) const override;
 	virtual bool ValidateFlattenedAdvancedCopyDestinations(const TMap<FString, FString>& FlattenedPackagesAndDestinations) const override;
diff --git a/Engine/Source/Developer/AssetTools/Public/IAssetTools.h b/Engine/Source/Developer/AssetTools/Public/IAssetTools.h
--- a/Engine/Source/Developer/AssetTools/Public/IAssetTools.h
+++ b/Engine/Source/Developer/AssetTools/Public/IAssetTools.h
@@ -14,6 +14,7 @@
 #include "AutomatedAssetImportData.h"
 #include "AssetRegistry/ARFilter.h"
 #include "Engine/Blueprint.h"
+#include "Logging/TokenizedMessage.h"
 #include "IAssetTools.generated.h"
 
 struct FAssetData;
@@ -476,7 +477,7 @@ public:
 	virtual bool AdvancedCopyPackages(const FAdvancedCopyParams& CopyParams, const TArray<TMap<FString, FString>> PackagesAndDestinations) const = 0;
 
 	/** Copies files after the flattened map of sources and destinations was confirmed */
-	virtual bool AdvancedCopyPackages(const TMap<FString, FString>& SourceAndDestPackages, const bool bForceAutosave = false, const bool bCopyOverAllDestinationOverlaps = true, FDuplicatedObjects* OutDuplicatedObjects = nullptr) const = 0;
+	virtual bool AdvancedCopyPackages(const TMap<FString, FString>& SourceAndDestPackages, const bool bForceAutosave = false, const bool bCopyOverAllDestinationOverlaps = true, FDuplicatedObjects* OutDuplicatedObjects = nullptr, EMessageSeverity::Type NotificationSeverityFilter = EMessageSeverity::Info) const = 0;
 
 	/* Given a set of packages to copy, generate the map of those packages to destination filenames */
 	virtual void GenerateAdvancedCopyDestinations(FAdvancedCopyParams& InParams, const TArray<FName>& InPackageNamesToCopy, const UAdvancedCopyCustomization* CopyCustomization, TMap<FString, FString>& OutPackagesAndDestinations) const = 0;
