Fix compile error in AdvancedCopyPackages; EMessageSeverity::CriticalError was removed in UE5

diff --git a/Engine/Source/Developer/AssetTools/Private/AssetTools.cpp b/Engine/Source/Developer/AssetTools/Private/AssetTools.cpp
--- a/Engine/Source/Developer/AssetTools/Private/AssetTools.cpp
+++ b/Engine/Source/Developer/AssetTools/Private/AssetTools.cpp
@@ -1153,7 +1153,7 @@ bool UAssetToolsImpl::AdvancedCopyPackages(
 			}
 		}
 		// @note Using the bForce param because the InSeverityFilter param is only checked against logs in the last page
-		AdvancedCopyLog.Notify(LogMessage, /*InSeverityFilter=*/EMessageSeverity::CriticalError, /*bForce=*/(Severity <= NotificationSeverityFilter));
+		AdvancedCopyLog.Notify(LogMessage, /*InSeverityFilter=*/EMessageSeverity::Error, /*bForce=*/(Severity <= NotificationSeverityFilter));
 		return true;
 	}
 	return false;
