Reworked the migration tool so that it support moving assets inside of a plugin.
Added a global event for when some package are migrated. This can be used for hotfixes and as extension point for systems to react or add stuff to the migration.

Other changes that were required for the new migration:
Level streaming is now aware that the world might be part of a instanced package. When it validate that file for the streamed world exist it will use the path from the package linker of the world package.
The level postload function now prefer to use the Instancing Context package remapping from its linker to determine where it should load some of its external UActorFolder.
The StringTableEditorModule was modified to be able to properly react to the migration.

This change will still require some change in a future release to reduce its memory usage.

diff --git a/Engine/Plugins/Editor/ContentBrowser/ContentBrowserAssetDataSource/Source/ContentBrowserAssetDataSource/Private/AssetFolderContextMenu.cpp b/Engine/Plugins/Editor/ContentBrowser/ContentBrowserAssetDataSource/Source/ContentBrowserAssetDataSource/Private/AssetFolderContextMenu.cpp
--- a/Engine/Plugins/Editor/ContentBrowser/ContentBrowserAssetDataSource/Source/ContentBrowserAssetDataSource/Private/AssetFolderContextMenu.cpp
+++ b/Engine/Plugins/Editor/ContentBrowser/ContentBrowserAssetDataSource/Source/ContentBrowserAssetDataSource/Private/AssetFolderContextMenu.cpp
@@ -212,7 +212,11 @@ void FAssetFolderContextMenu::ExecuteMigrateFolder()
 		AssetViewUtils::GetAssetsInPaths(SelectedPaths, AssetDataList);
 		for ( auto AssetIt = AssetDataList.CreateConstIterator(); AssetIt; ++AssetIt )
 		{
-			PackageNames.Add((*AssetIt).PackageName);
+			// Don't collect the external packages here
+			if (AssetIt->GetOptionalOuterPathName().IsNone())
+			{
+				PackageNames.Add((*AssetIt).PackageName);
+			}
 		}
 
 		PackageNames.Append(SelectedPackages);
diff --git a/Engine/Source/Developer/AssetTools/AssetTools.Build.cs b/Engine/Source/Developer/AssetTools/AssetTools.Build.cs
--- a/Engine/Source/Developer/AssetTools/AssetTools.Build.cs
+++ b/Engine/Source/Developer/AssetTools/AssetTools.Build.cs
@@ -45,13 +45,13 @@ public class AssetTools : ModuleRules
 				"InterchangeCore",
 				"InterchangeEngine",
 				"PhysicsUtilities",
+				"AssetRegistry",
 			}
 		);
 
 		PrivateIncludePathModuleNames.AddRange(
 			new string[] {
 				"Analytics",
-				"AssetRegistry",
 				"ContentBrowser",
 				"CollectionManager",
                 "CurveAssetEditor",
@@ -86,7 +86,6 @@ public class AssetTools : ModuleRules
 
 		DynamicallyLoadedModuleNames.AddRange(
 			new string[] {
-				"AssetRegistry",
 				"ContentBrowser",
 				"CollectionManager",
 				"CurveTableEditor",
diff --git a/Engine/Source/Developer/AssetTools/Private/AssetTools.cpp b/Engine/Source/Developer/AssetTools/Private/AssetTools.cpp
--- a/Engine/Source/Developer/AssetTools/Private/AssetTools.cpp
+++ b/Engine/Source/Developer/AssetTools/Private/AssetTools.cpp
@@ -7,6 +7,7 @@
 #include "HAL/FileManager.h"
 #include "Misc/FileHelper.h"
 #include "Misc/Paths.h"
+#include "Misc/PathViews.h"
 #include "Misc/ScopedSlowTask.h"
 #include "UObject/GCObjectScopeGuard.h"
 #include "UObject/UObjectHash.h"
@@ -161,16 +162,714 @@
 #include "Widgets/Layout/SSpacer.h"
 #include "Interfaces/IPluginManager.h"
 #include "Settings/ContentBrowserSettings.h"
+#include "Algo/Count.h"
+#include "AssetRegistry/AssetRegistryModule.h"
+#include "AssetRegistry/IAssetRegistry.h"
+#include "WorldPartition/WorldPartitionActorDescUtils.h"
+#include "UObject/NameTypes.h"
+#include "ExternalPackageHelper.h"
+#include "UObject/LinkerInstancingContext.h"
+#include "PackageMigrationContext.h"
 
 #if WITH_EDITOR
 #include "Subsystems/AssetEditorSubsystem.h"
 #include "Settings/EditorExperimentalSettings.h"
+#include "ObjectTools.h"
 #endif
 
 #include UE_INLINE_GENERATED_CPP_BY_NAME(AssetTools)
 
 #define LOCTEXT_NAMESPACE "AssetTools"
 
+namespace UE::AssetTools::Private
+{
+	bool bUseNewPackageMigration = true;
+	static FAutoConsoleVariableRef CVarUseNewPackageMigration(
+		TEXT("AssetTools.UseNewPackageMigration"),
+		bUseNewPackageMigration,
+		TEXT("When set, The package migration will use the new implementation made for 5.1.")
+	);
+
+	// use a struct as a namespace to allow easier friend declarations
+	struct FPackageMigrationImpl
+	{
+	public:
+		struct FPackageMigrationImplContext
+		{
+			FPackageMigrationImplContext()
+				: SlowTask(100, LOCTEXT("MigratePackages", "Migrating Packages"))
+			{
+				InstancingContext.SetSoftObjectPathRemappingEnabled(true);
+				InstancingContext.AddTag(ULevel::DontLoadExternalObjectsTag);
+
+				const bool bCanCancel = true;
+				SlowTask.MakeDialog(bCanCancel);
+			}
+
+			bool bWasCanceled = false;
+			FLinkerInstancingContext InstancingContext;
+			FScopedSlowTask SlowTask;
+			FMigrationOptions Options;
+		};
+
+
+		static FString GetMountPointRootPath(bool bIsDestinationAProject, const FString& DestinationFolder)
+		{
+			FString MountPointRootPath;
+			if (bIsDestinationAProject)
+			{
+				MountPointRootPath = TEXT("/Game/");
+			}
+			else
+			{
+				TArray<FString> PluginFiles;
+				IFileManager::Get().FindFiles(PluginFiles, *(DestinationFolder + TEXT("../")), TEXT("uplugin"));
+
+				if (PluginFiles.Num() == 1)
+				{
+					FStringView PluginNameView = FPathViews::GetBaseFilename(PluginFiles[0]);
+					MountPointRootPath.Reserve(PluginNameView.Len() + 2);
+					MountPointRootPath.AppendChar(TEXT('/'));
+					MountPointRootPath.Append(PluginNameView);
+					MountPointRootPath.AppendChar(TEXT('/'));
+				}
+				else
+				{
+					UE_LOG(LogAssetTools
+						, Error
+						, TEXT("{%s} does not appear to be a game Content folder. Migrated content only work if placed in a Content folder. Aborting the Migration.")
+						, *DestinationFolder
+					);
+				}
+			}
+
+			return MountPointRootPath;
+		}
+
+		static UObject* FindAssetInPackage(UPackage* Package)
+		{
+			// Inspired from the save code (UEditorEngine::Save)
+			UObject* Asset = nullptr;
+			if (Package->HasAnyPackageFlags(PKG_ContainsMap))
+			{
+				Asset = UWorld::FindWorldInPackage(Package);
+			}
+			
+			if (!Asset)
+			{
+				// Otherwise find the main asset of the package
+				Asset = Package->FindAssetInPackage();
+			}
+
+			return Asset;
+		};
+
+		// Prepare the MigrationPackagesData in the migration context and log any info about the package that will be renamed
+		static void SetupPublicAssetPackagesMigrationData(const TSharedPtr<TArray<ReportPackageData>>& PackageDataToMigrate
+			, TArray<TPair<const ReportPackageData*, const FAssetData>>& ExternalPackageDatas
+			, TMap<const FStringView, int32>& ExistingPackageNameToMigrationDataIndex
+			, FPackageMigrationContext& PackageMigrationContext
+			, FPackageMigrationImplContext& PackageMigrationImplContext
+			, FMessageLog& MigrateLog)
+		{
+			FAssetRegistryModule& AssetRegistryModule = FModuleManager::LoadModuleChecked<FAssetRegistryModule>(TEXT("AssetRegistry"));
+			IAssetRegistry& AssetRegistry = AssetRegistryModule.Get();
+			TArray<FAssetData> AssetsData;
+
+			// Used to avoid a name collision between the instanced packages
+			TSet<FString> UsedPackagesName;
+			UsedPackagesName.Reserve(PackageDataToMigrate->Num());
+			IFileManager& FileManager = IFileManager::Get();
+
+
+			ExternalPackageDatas.Reserve(PackageDataToMigrate->Num());
+			ExistingPackageNameToMigrationDataIndex.Reserve(PackageDataToMigrate->Num());
+
+			TArray<FText> RenamedPackagesInfo;
+
+
+			// Handle prompt response and automation
+			EAppReturnType::Type LastResponse = EAppReturnType::Yes;
+			if (!PackageMigrationImplContext.Options.bPrompt)
+			{
+				switch (PackageMigrationImplContext.Options.AssetConflict)
+				{
+				case EAssetMigrationConflict::Overwrite:
+					LastResponse = EAppReturnType::YesAll;
+					break;
+				case EAssetMigrationConflict::Skip:
+					LastResponse = EAppReturnType::NoAll;
+					break;
+				case EAssetMigrationConflict::Cancel:
+					PackageMigrationImplContext.bWasCanceled = true;
+					break;
+				}
+			}
+
+			// Process the normal package here (no external object or actors).
+			for (const ReportPackageData& PackageData : *(PackageDataToMigrate.Get()))
+			{
+				if (PackageMigrationImplContext.SlowTask.ShouldCancel() || PackageMigrationImplContext.bWasCanceled)
+				{
+					PackageMigrationImplContext.bWasCanceled = true;
+					break;
+				}
+
+				if (!PackageData.bShouldMigratePackage)
+				{
+					continue;
+				}
+
+				AssetsData.Reset();
+				const bool bIncludeOndiskAssetOnly = true;
+				// We want all the asset including those not supported for the current platform.
+				const bool bSkipARFilteredAssets = false;
+				AssetRegistry.GetAssetsByPackageName(*PackageData.Name, AssetsData, bIncludeOndiskAssetOnly, bSkipARFilteredAssets);
+
+				// Detect if the package contains some external object or actors
+				{
+					bool bIsExternalPackage = false;
+					for (FAssetData& AssetData : AssetsData)
+					{
+						if (!AssetData.GetOptionalOuterPathName().IsNone())
+						{
+							bIsExternalPackage = true;
+							ExternalPackageDatas.Emplace(&PackageData, MoveTemp(AssetData));
+							break;
+						}
+					}
+
+					if (bIsExternalPackage)
+					{
+						// Skip this asset since it is a external package. They are done in another pass
+						continue;
+					}
+				}
+
+
+				const FPackageMigrationContext::FScopedMountPoint& MountPoint = PackageMigrationContext.GetDestinationMountPoint();
+
+				FString NewPackageName = MountPoint.GetNewPackageNameForMigration(PackageData.Name);
+
+				// Is there a simpler way to get the package extension?
+				FString ExistingFilename;
+				FPackageName::DoesPackageExist(PackageData.Name, &ExistingFilename);
+
+				const bool bIncludeDot = true;
+				FString NewPackageFilename = MountPoint.GetMigratedPackageFilename(NewPackageName, FPathViews::GetExtension(ExistingFilename, bIncludeDot));
+				uint32 NewPackageNameHash = GetTypeHash(NewPackageName);
+
+				// Check we can use the new name for the asset
+				const bool bIsNameChangeCausedByOtherMigrationPackage = UsedPackagesName.ContainsByHash(NewPackageNameHash, NewPackageName);
+				const bool bIsNameChangeCausedByExistingAssetInDestination = FileManager.FileExists(*NewPackageFilename);
+
+				bool bShouldMigrate = true;
+
+				// Ask the user what to do if an asset already exist in the destination
+				if (bIsNameChangeCausedByExistingAssetInDestination)
+				{
+					// Handle name collision in the destination project. Post 5.1 should expose the possibility to rename the migrated package.
+
+					EAppReturnType::Type Response;
+					if (FApp::IsUnattended() || !PackageMigrationImplContext.Options.bPrompt || LastResponse == EAppReturnType::YesAll || LastResponse == EAppReturnType::NoAll)
+					{
+						Response = LastResponse;
+					}
+					else
+					{
+						const FText Message = FText::Format(LOCTEXT("MigratePackages_AlreadyExists", "An asset already exists at location {0} would you like to overwrite it?"), FText::FromString(NewPackageFilename));
+						Response = FMessageDialog::Open(EAppMsgType::YesNoYesAllNoAllCancel, Message);
+						LastResponse = Response;
+					}
+
+					if (Response == EAppReturnType::Cancel)
+					{
+						// The user chose to cancel mid-operation. Break out.
+						PackageMigrationImplContext.bWasCanceled = true;
+						break;
+					}
+
+					if (Response == EAppReturnType::No || Response == EAppReturnType::NoAll)
+					{
+						bShouldMigrate = false;
+					}
+				}
+
+				// Avoid a name collision
+				if (bIsNameChangeCausedByOtherMigrationPackage)
+				{
+					// Name collision with another package to migrate (simply add _number to package name)
+					int32 Index = 1;
+					FString ModifiedNewPackageName;
+					// 4 for "_num"
+					ModifiedNewPackageName.Reserve(NewPackageName.Len() + 3);
+					uint32 ModifiedNewPackageNameHash;
+
+					FString ModifiedNewPackageFilename;
+					ModifiedNewPackageFilename.Reserve(NewPackageFilename.Len() + 3);
+					FStringView FileExtension = FPathViews::GetExtension(NewPackageFilename, bIncludeDot);
+					FStringView NewPackageFilenameWithoutExtension(*NewPackageFilename, NewPackageFilename.Len() - FileExtension.Len());
+
+					do
+					{
+						++Index;
+						ModifiedNewPackageName.Reset();
+						ModifiedNewPackageName.Append(NewPackageName);
+						ModifiedNewPackageName.AppendChar('_');
+						ModifiedNewPackageName.AppendInt(Index);
+						ModifiedNewPackageNameHash = GetTypeHash(ModifiedNewPackageName);
+
+						ModifiedNewPackageFilename.Reset();
+						ModifiedNewPackageFilename.Append(NewPackageFilenameWithoutExtension);
+						ModifiedNewPackageFilename.AppendChar('_');
+						ModifiedNewPackageFilename.AppendInt(Index);
+						ModifiedNewPackageFilename.Append(FileExtension);
+					} while (UsedPackagesName.ContainsByHash(ModifiedNewPackageNameHash, ModifiedNewPackageName)
+						|| FileManager.FileExists(*ModifiedNewPackageFilename)
+						);
+
+					if (bIsNameChangeCausedByOtherMigrationPackage)
+					{
+						RenamedPackagesInfo.Add(
+							FText::Format(LOCTEXT("MigratePackages_PackageNameChangeByOtherPackage", "({0}) will be named ({1}) because another migrated package is already using that name.")
+								, FText::FromString(PackageData.Name)
+								, FText::FromString(ModifiedNewPackageName)
+							)
+						);
+					}
+					else if (bIsNameChangeCausedByExistingAssetInDestination)
+					{
+						RenamedPackagesInfo.Add(
+							FText::Format(LOCTEXT("MigratePackages_PackageNameChangeByExistingPackageInDestination", "({0}) will be named ({1}) because an file in the destination project is already using that name.")
+								, FText::FromString(PackageData.Name)
+								, FText::FromString(ModifiedNewPackageName)
+							)
+						);
+					}
+
+					NewPackageNameHash = ModifiedNewPackageNameHash;
+					NewPackageName = MoveTemp(ModifiedNewPackageName);
+					NewPackageFilename = MoveTemp(ModifiedNewPackageFilename);
+				}
+
+				UsedPackagesName.AddByHash(NewPackageNameHash, NewPackageName);
+
+				int32 MigrationPackageDataIndex = PackageMigrationContext.MigrationPackagesData.Emplace(NewPackageName, PackageData.Name, NewPackageFilename);
+				FPackageMigrationContext::FMigrationPackageData& MigrationData = PackageMigrationContext.MigrationPackagesData[MigrationPackageDataIndex];
+
+				// Post 5.1 look into detecting which file can be migrated without being loaded
+				MigrationData.bNeedInstancedLoad = bShouldMigrate;
+				MigrationData.bNeedToBeSaveMigrated = bShouldMigrate;
+
+				ExistingPackageNameToMigrationDataIndex.Add(PackageData.Name, MigrationPackageDataIndex);
+
+				// Check if there is a existing package that need to be moved out the way
+				if (UPackage* InTheWayPacakge = FindObjectFast<UPackage>(nullptr, *NewPackageName))
+				{
+					PackageMigrationContext.MoveInTheWayPackage(InTheWayPacakge);
+				}
+
+				PackageMigrationImplContext.InstancingContext.AddPackageMapping(*PackageData.Name, *NewPackageName);
+			}
+
+			if (!RenamedPackagesInfo.IsEmpty())
+			{
+				MigrateLog.NewPage(LOCTEXT("MigratePackages_PackageNameChange", "Modified Packages Name For Migration"));
+				for (const FText& Message : RenamedPackagesInfo)
+				{
+					MigrateLog.Info(Message);
+				}
+			}
+		}
+
+		static void SetupExternalAssetPackagesMigrationData(const TArray<TPair<const ReportPackageData*, const FAssetData>>& ExternalPackageDatas, TMap<const FStringView, int32>& ExistingPackageNameToMigrationDataIndex, FPackageMigrationContext& PackageMigrationContext, FPackageMigrationImplContext& MigrationImplContext)
+		{
+			auto GenerateObjectPath = [](const FStringView PackageName, const FStringView AssetName, const FStringView SubPath) -> FString
+			{
+				FString InstancedPath;
+
+				int32 DelimitersCount = 2;
+				if (SubPath.IsEmpty())
+				{
+					--DelimitersCount;
+				}
+
+				InstancedPath.Reserve(PackageName.Len() + AssetName.Len() + SubPath.Len() + DelimitersCount);
+				InstancedPath.Append(PackageName);
+				InstancedPath.AppendChar(TEXT('.'));
+				InstancedPath.Append(AssetName);
+
+				if (!SubPath.IsEmpty())
+				{
+					InstancedPath.AppendChar(SUBOBJECT_DELIMITER_CHAR);
+					InstancedPath.Append(SubPath);
+				}
+
+				return InstancedPath;
+			};
+
+			// Handle the external packages
+			for (const TPair<const ReportPackageData*, const FAssetData>& PackageData : ExternalPackageDatas)
+			{
+				const FString OuterPath = PackageData.Value.GetOptionalOuterPathName().ToString();
+
+				const FStringView OuterPackageName = FPathViews::GetBaseFilenameWithPath(OuterPath);
+
+				const int32* PtrToMigrationDataIndex = ExistingPackageNameToMigrationDataIndex.Find(OuterPackageName);
+
+				if (!PtrToMigrationDataIndex)
+				{
+					PackageMigrationContext.AddErrorMigrationMessage(FText::Format(
+						LOCTEXT("MigratePackages_ExternalPackageNotExported", "({0}) won't be migrated because the outer of its content ({1}) will not be migrated")
+						, FText::FromString(PackageData.Key->Name)
+						, FText::FromStringView(OuterPackageName)
+					)
+					);
+					continue;
+				}
+
+				FPackageMigrationContext::FMigrationPackageData& OuterMigrationData = PackageMigrationContext.MigrationPackagesData[*PtrToMigrationDataIndex];
+
+				const FSoftObjectPath ExistingObjectPath = PackageData.Value.GetSoftObjectPath();
+				FString NewObjectPath = GenerateObjectPath(OuterMigrationData.GetInstancedPackageName(), FPathViews::GetBaseFilename(OuterMigrationData.GetInstancedPackageName()), ExistingObjectPath.GetSubPathString());
+
+				FString NewPackageName;
+
+				if (FWorldPartitionActorDescUtils::IsValidActorDescriptorFromAssetData(PackageData.Value))
+				{
+					// If the external package is an actor.
+					EActorPackagingScheme Scheme = ULevel::GetActorPackagingSchemeFromActorPackageName(PackageData.Key->Name);
+					NewPackageName = ULevel::GetActorPackageName(ULevel::GetExternalActorsPath(OuterMigrationData.GetInstancedPackageName()), Scheme, NewObjectPath);
+				}
+				else
+				{
+					NewPackageName = FExternalPackageHelper::GetExternalPackageName(OuterMigrationData.GetInstancedPackageName(), NewObjectPath);
+				}
+
+
+				FString NewPackageFilename = PackageMigrationContext.GetDestinationMountPoint().GetMigratedPackageFilename(NewPackageName, FPackageName::GetAssetPackageExtension());
+
+				// Add the external package to the list of migrated asset, even if the asset is not really migrated we still need it for the instance path
+				uint32 MigrationPackageDataIndex = PackageMigrationContext.MigrationPackagesData.Emplace(NewPackageName, PackageData.Key->Name, NewPackageFilename);
+
+				FPackageMigrationContext::FMigrationPackageData& MigrationData = PackageMigrationContext.MigrationPackagesData[MigrationPackageDataIndex];
+
+				// Post 5.1 look into detecting which file can be migrated without being loaded
+				MigrationData.bNeedInstancedLoad = OuterMigrationData.bNeedInstancedLoad && PackageData.Key->bShouldMigratePackage;
+				MigrationData.bNeedToBeSaveMigrated = OuterMigrationData.bNeedToBeSaveMigrated && PackageData.Key->bShouldMigratePackage;
+
+				// Check if there is a existing package that need to be moved out the way
+				if (UPackage* InTheWayPacakge = FindObjectFast<UPackage>(nullptr, *NewPackageName))
+				{
+					PackageMigrationContext.MoveInTheWayPackage(InTheWayPacakge);
+				}
+
+				MigrationImplContext.InstancingContext.AddPackageMapping(*(PackageData.Key->Name), *NewPackageName);
+			}
+		}
+
+		static void CreatePackagesAndSetupLinkers(FPackageMigrationContext& PackageMigrationContext, FPackageMigrationImplContext& MigrationImplContext, TArray<TWeakObjectPtr<UPackage>>& PackagesToClean)
+		{
+			FArchive* ReaderOverride = nullptr;
+
+			// Create the package and their load linker. If an asset try to load a weak dependency the package will load from the right file because of the existing linker.
+			for (FPackageMigrationContext::FMigrationPackageData& MigrationPackageData : PackageMigrationContext.MigrationPackagesData)
+			{
+				if (MigrationImplContext.SlowTask.ShouldCancel() || MigrationImplContext.bWasCanceled)
+				{
+					MigrationImplContext.bWasCanceled = true;
+					break;
+				}
+
+				UPackage* MigrationPackage = CreatePackage(*MigrationPackageData.GetInstancedPackageName());
+
+				/**
+				 * Load_Verify tell the linker to not load the package but it will create the linker if the file exist and is valid.
+				 * LOAD_NoVerify tell the linker to not check the import of the package (in the editor this will avoid loading the hard dependencies of the package)
+				 * 
+				 * This will the package with a pre-created linker that has the right instancing context. Ready to loaded by another call to load package that may not have all that info.
+				 */
+				MigrationPackageData.InstancedPackage = LoadPackage(MigrationPackage, *MigrationPackageData.GetOriginalPackageName(), LOAD_Verify | LOAD_NoVerify, ReaderOverride, &MigrationImplContext.InstancingContext);
+
+				PackagesToClean.Add(MigrationPackage);
+			}
+		}
+
+		static void LoadInstancedPackages(FPackageMigrationContext& PackageMigrationContext, FPackageMigrationImplContext& MigrationImplContext)
+		{
+			const int32 NumOfPackageToLoad = Algo::CountIf(PackageMigrationContext.MigrationPackagesData, [](const FPackageMigrationContext::FMigrationPackageData& MigrationPackageData)
+				{
+					return MigrationPackageData.bNeedInstancedLoad || MigrationPackageData.bNeedToBeSaveMigrated;
+				});
+
+			// We give 40% of the progress weight to the loading
+			const float ProgressPerItemToload = 40.f / NumOfPackageToLoad;
+
+			uint32 ProgressCount = 0;
+
+			FArchive* ReaderOverride = nullptr;
+
+			for (FPackageMigrationContext::FMigrationPackageData& MigrationPackageData : PackageMigrationContext.MigrationPackagesData)
+			{
+				if (MigrationImplContext.SlowTask.ShouldCancel() || MigrationImplContext.bWasCanceled)
+				{
+					MigrationImplContext.bWasCanceled = true;
+					break;
+				}
+
+				if (MigrationPackageData.bNeedInstancedLoad || MigrationPackageData.bNeedToBeSaveMigrated)
+				{
+					MigrationImplContext.SlowTask.EnterProgressFrame(ProgressPerItemToload
+						, FText::Format(LOCTEXT("MigratePackages_Load", "Migration: Loading packages {0}/{1}")
+							, ProgressCount
+							, NumOfPackageToLoad
+						)
+					);
+
+					++ProgressCount;
+
+					UPackage* LoadedPackage = LoadPackage(MigrationPackageData.GetInstancedPackage(), *MigrationPackageData.GetOriginalPackageName(), LOAD_None, ReaderOverride, &MigrationImplContext.InstancingContext);
+
+					if (IsValid(LoadedPackage))
+					{
+						UObject* Asset = FPackageMigrationImpl::FindAssetInPackage(LoadedPackage);
+
+						if (Asset && !Asset->HasAnyFlags(RF_Standalone))
+						{
+							// The external actor or object are sometimes not protected from the garbage collector. This is removed by the purge at the end of the migration
+							Asset->AddToRoot();
+						}
+					}
+					else
+					{
+						MigrationPackageData.bNeedToBeSaveMigrated = false;
+						PackageMigrationContext.AddErrorMigrationMessage(FText::Format(LOCTEXT("MigratePackages_LoadFailed", "Couldn't migrate package ({0}) because the load of the asset failed"), FText::FromString(MigrationPackageData.GetOriginalPackageName())));
+					}
+				}
+			}
+		}
+
+		static void SaveInstancedPackagesIntoDestination(FPackageMigrationContext& PackageMigrationContext, FPackageMigrationImplContext& MigrationImplContext)
+		{
+			FSavePackageArgs SaveArgs;
+			SaveArgs.SaveFlags |= SAVE_RehydratePayloads | SAVE_Async;
+
+			uint32 ProgressCount = 0;
+
+			const int32 NumOfPackageToSave = Algo::CountIf(PackageMigrationContext.MigrationPackagesData, [](const FPackageMigrationContext::FMigrationPackageData& MigrationPackageData)
+				{
+					return MigrationPackageData.bNeedToBeSaveMigrated;
+				});
+
+			// We give 60% of the progress weight to the saving of the assets
+			const float ProgressPerItemToSave = 60.f / NumOfPackageToSave;
+
+			for (const FPackageMigrationContext::FMigrationPackageData& MigrationPackageData : PackageMigrationContext.MigrationPackagesData)
+			{
+				if (MigrationImplContext.SlowTask.ShouldCancel() || MigrationImplContext.bWasCanceled)
+				{
+					MigrationImplContext.bWasCanceled = true;
+					break;
+				}
+				if (MigrationPackageData.bNeedToBeSaveMigrated)
+				{
+
+					MigrationImplContext.SlowTask.EnterProgressFrame(ProgressPerItemToSave
+						, FText::Format(LOCTEXT("MigratePackages_Save", "Saving the packages into the destination {0}/{1}")
+							, ProgressCount
+							, NumOfPackageToSave
+						)
+					);
+
+					++ProgressCount;
+					
+					UObject* Asset = FPackageMigrationImpl::FindAssetInPackage(MigrationPackageData.GetInstancedPackage());
+					if (IsValid(MigrationPackageData.GetInstancedPackage()) && Asset)
+					{
+						FSavePackageResultStruct SaveResult = GEditor->Save(MigrationPackageData.GetInstancedPackage(), Asset, *MigrationPackageData.GetDestinationFilename(), SaveArgs);
+
+						if (SaveResult.IsSuccessful())
+						{
+							PackageMigrationContext.AddSucessfullMigrationMessage(FText::Format(LOCTEXT("MigratePackages_SaveSuccess", "Package ({0}) was migrated successfully migrated as ({1}) with the fellowing filename ({2})")\
+								, FText::FromString(MigrationPackageData.GetOriginalPackageName())
+								, FText::FromString(MigrationPackageData.GetInstancedPackageName())
+								, FText::FromString(MigrationPackageData.GetDestinationFilename())
+							));
+						}
+						else
+						{
+							PackageMigrationContext.AddErrorMigrationMessage(FText::Format(LOCTEXT("MigratePackages_SaveFailed", "Couldn't migrate package ({0}) because the asset save failed. Destination File ({1})")
+								, FText::FromString(MigrationPackageData.GetOriginalPackageName())
+								, FText::FromString(MigrationPackageData.GetDestinationFilename())
+							));
+						}
+					}
+					else if (!IsValid(MigrationPackageData.GetInstancedPackage()))
+					{
+						PackageMigrationContext.AddErrorMigrationMessage(FText::Format(LOCTEXT("MigratePackages_SaveFailedPackageInvalid", "Couldn't migrate package ({0}) because the asset package is invalid. Destination File ({1})")
+							, FText::FromString(MigrationPackageData.GetOriginalPackageName())
+							, FText::FromString(MigrationPackageData.GetDestinationFilename())
+						));
+					}
+					else
+					{
+						PackageMigrationContext.AddErrorMigrationMessage(FText::Format(LOCTEXT("MigratePackages_SaveFailedNoAsset", "Couldn't migrate package ({0}) because the package didn't contains a an asset. Destination File ({1})")
+							, FText::FromString(MigrationPackageData.GetOriginalPackageName())
+							, FText::FromString(MigrationPackageData.GetDestinationFilename())
+						));
+					}
+				}
+			}
+		}
+
+		static void CleanInstancedPackages(const TArray<TWeakObjectPtr<UPackage>>& PackagesToClean)
+		{
+			TSharedRef<TFunction<void(UObject*)>> PurgeObject = MakeShared<TFunction<void(UObject*)>>();
+			PurgeObject.Get() = [PurgeObject](UObject* Object)
+			{
+				if (IsValid(Object))
+				{
+					if (UPackage* ExternalPackage = Object->GetExternalPackage())
+					{
+						if (Object != ExternalPackage && IsValid(ExternalPackage))
+						{
+							(*PurgeObject)(ExternalPackage);
+							ForEachObjectWithOuter(ExternalPackage, *PurgeObject);
+						}
+					}
+
+					if (Object->IsRooted())
+					{
+						Object->RemoveFromRoot();
+					}
+
+					Object->ClearFlags(RF_Public | RF_Standalone);
+					Object->MarkAsGarbage();
+				}
+			};
+
+			// Turn off the components while unloading stuff
+			FGlobalComponentReregisterContext ComponentContext;
+
+			TArray<UObject*> ReferenceToNull;
+			for (const TWeakObjectPtr<UPackage>& WeakPackage : PackagesToClean)
+			{
+				if (UPackage* Package = WeakPackage.Get())
+				{
+					const ERenameFlags PkgRenameFlags = REN_ForceNoResetLoaders | REN_DoNotDirty | REN_DontCreateRedirectors | REN_NonTransactional | REN_SkipGeneratedClasses;
+					check(Package->Rename(*MakeUniqueObjectName(nullptr, UPackage::StaticClass(), *FString::Printf(TEXT("%s_DEADFROMMIGRATION"), *Package->GetName())).ToString(), nullptr, PkgRenameFlags));
+
+					if (UObject* Asset = FPackageMigrationImpl::FindAssetInPackage(Package))
+					{
+						ReferenceToNull.Add(Asset);
+					}
+
+					(*PurgeObject)(Package);
+					ForEachObjectWithOuter(Package, *PurgeObject);
+				}
+			}
+
+			// Removing assets from memory is complicated
+			TArray<UObject*> UDStructToReplace;
+			for (int32 Index = 0; Index < ReferenceToNull.Num(); )
+			{
+				if (UUserDefinedStruct* UDStruct = Cast<UUserDefinedStruct>(ReferenceToNull[Index]))
+				{
+					ReferenceToNull.RemoveAtSwap(Index);
+					UDStructToReplace.Add(UDStruct);
+				}
+				else
+				{
+					Index++;
+				}
+			}
+
+			if (UDStructToReplace.Num())
+			{
+				ObjectTools::ForceReplaceReferences(GetFallbackStruct(), UDStructToReplace);
+			}
+
+			ObjectTools::ForceReplaceReferences(nullptr, ReferenceToNull);
+
+			CollectGarbage(GARBAGE_COLLECTION_KEEPFLAGS);
+		}
+
+		static void RestoreInTheWayPackages(FPackageMigrationContext& PackageMigrationContext)
+		{
+			PackageMigrationContext.TemporalyMovedPackages.Empty();
+
+			CollectGarbage(GARBAGE_COLLECTION_KEEPFLAGS);
+		}
+
+		static void BuildAndNotifyLogWindow(FPackageMigrationContext& PackageMigrationContext, FMessageLog& MigrateLog, FPackageMigrationImplContext& MigrationImplContext)
+		{
+			FText NotificationMessage;
+
+			if (!PackageMigrationContext.MigratedPackageMessages.IsEmpty())
+			{
+				MigrateLog.NewPage(LOCTEXT("MigratePackages_SuccessPageLabel", "Migrated Pacakges"));
+				for (const FText& InfoMessage : PackageMigrationContext.MigratedPackageMessages)
+				{
+					MigrateLog.Info(InfoMessage);
+				}
+			}
+
+			if (!PackageMigrationContext.WarningMessage.IsEmpty())
+			{
+				MigrateLog.NewPage(LOCTEXT("MigratePackages_WarningPageLabel", "Migration warnings"));
+				for (const FText& WarningMessage : PackageMigrationContext.WarningMessage)
+				{
+					MigrateLog.Warning(WarningMessage);
+				}
+			}
+
+			if (PackageMigrationContext.ErrorMessages.IsEmpty())
+			{
+				if (PackageMigrationContext.WarningMessage.IsEmpty())
+				{
+					NotificationMessage = LOCTEXT("MigratePackages_SuccessNotifcation", "Content migration completed.");
+				}
+				else
+				{
+					NotificationMessage = LOCTEXT("MigratePackages_SuccessNotifcation", "Content migration completed with some warnings.");
+				}
+			}
+			else
+			{
+				if (PackageMigrationContext.MigratedPackageMessages.IsEmpty())
+				{
+					NotificationMessage = LOCTEXT("MigratePackages_FailureNotifcation", "Content migration failed.");
+				}
+				else
+				{
+					NotificationMessage = LOCTEXT("MigratePackages_PartialFailureNotifcation", "Content migration failed but some asset have been migrated sucessfuly.");
+				}
+
+				MigrateLog.NewPage(LOCTEXT("MigratePackages_ErrorsPageLabel", "Migration Errors"));
+				for (const FText& ErrorMessage : PackageMigrationContext.ErrorMessages)
+				{
+					MigrateLog.Error(ErrorMessage);
+				}
+			}
+
+			if (MigrationImplContext.bWasCanceled)
+			{
+				if (PackageMigrationContext.MigratedPackageMessages.IsEmpty())
+				{
+					NotificationMessage = LOCTEXT("MigratePackages_Canceled", "Content migration was canceled.");
+				}
+				else
+				{
+					NotificationMessage = LOCTEXT("MigratePackages_Canceled", "Content migration was canceled but some package where already migrated.");
+				}
+			}
+
+			MigrateLog.Notify(NotificationMessage, EMessageSeverity::Info);
+		}
+	};
+}
+
+
 IAssetTools& IAssetTools::Get()
 {
 	FAssetToolsModule& Module = FModuleManager::LoadModuleChecked<FAssetToolsModule>("AssetTools");
@@ -1733,6 +2432,11 @@ void UAssetToolsImpl::MigratePackages(const TArray<FName>& PackageNamesToMigrate
 	}
 }
 
+UE::AssetTools::FOnPackageMigration& UAssetToolsImpl::GetOnPackageMigration()
+{
+	return OnPackageMigration;
+}
+
 void UAssetToolsImpl::OnNewImportRecord(UClass* AssetType, const FString& FileExtension, bool bSucceeded, bool bWasCancelled, const FDateTime& StartTime)
 {
 	// Don't attempt to report usage stats if analytics isn't available
@@ -3074,7 +3778,7 @@ void UAssetToolsImpl::PerformMigratePackages(TArray<FName> PackageNamesToMigrate
 	}
 }
 
-void UAssetToolsImpl::MigratePackages_ReportConfirmed(TSharedPtr<TArray<ReportPackageData>> PackageDataToMigrate, const FString DestinationPath, const struct FMigrationOptions Options) const
+void UAssetToolsImpl::MigratePackages_ReportConfirmed(TSharedPtr<TArray<ReportPackageData>> PackageDataToMigrate, const FString DestinationPath, const FMigrationOptions Options) const
 {
 	FString DestinationFolder;
 	if (FApp::IsUnattended() || !DestinationPath.IsEmpty())
@@ -3128,10 +3832,13 @@ void UAssetToolsImpl::MigratePackages_ReportConfirmed(TSharedPtr<TArray<ReportPa
 				}
 				else
 				{
-					// The user chose a non-content folder. Confirm that this was their intention.
-					const FText Message = FText::Format(LOCTEXT("MigratePackages_NonContentFolder", "{0} does not appear to be a game Content folder. Migrated content will only work properly if placed in a Content folder. Would you like to place your content here anyway?"), FText::FromString(DestinationFolder));
-					EAppReturnType::Type Response = FMessageDialog::Open(EAppMsgType::YesNo, Message);
-					bFolderAccepted = (Response == EAppReturnType::Yes);
+					// The user chose a non-content folder. Let them know they cannot do that.
+					const FText Message = FText::Format(LOCTEXT("MigratePackages_NonContentFolder", "{0} does not appear to be a game Content folder. Migrated content only work if placed in a Content folder. Select a Content folder."), FText::FromString(DestinationFolder));
+					EAppReturnType::Type Response = FMessageDialog::Open(EAppMsgType::OkCancel, Message);
+					if (Response == EAppReturnType::Cancel)
+					{
+						return;
+					}
 				}
 			}
 		}
@@ -3153,7 +3860,6 @@ void UAssetToolsImpl::MigratePackages_ReportConfirmed(TSharedPtr<TArray<ReportPa
 		}
 	}
 
-	bool bAbort = false;
 	FMessageLog MigrateLog("AssetTools");
 
 	// Determine if the destination is a project content folder or a plugin content folder
@@ -3161,550 +3867,663 @@ void UAssetToolsImpl::MigratePackages_ReportConfirmed(TSharedPtr<TArray<ReportPa
 	IFileManager::Get().FindFiles(ProjectFiles, *(DestinationFolder + TEXT("../")), TEXT("uproject"));
 	bool bIsDestinationAProject = !ProjectFiles.IsEmpty();
 
-	// Associate each Content folder in the target Plugin hierarchy to a content root string in UFS
-	TMap<FName, FString> DestContentRootsToFolders;
-
-	// Assets in /Game always map directly to the destination
-	DestContentRootsToFolders.Add(FName(TEXT("/Game")), DestinationFolder);
-
-	// If our destination is a project, it could have plugins...
-	if (bIsDestinationAProject)
+	if (UE::AssetTools::Private::bUseNewPackageMigration)
 	{
-		// Find all "Content" folders under the destination ../Plugins directory
-		TArray<FString> ContentFolders;
-		IFileManager::Get().FindFilesRecursive(ContentFolders, *(DestinationFolder + TEXT("../Plugins/")), TEXT("Content"), false, true);
+		using namespace UE::AssetTools::Private;
+
+		FString MountPointRootPath = FPackageMigrationImpl::GetMountPointRootPath(bIsDestinationAProject, DestinationFolder);
 
-		for (const FString& Folder : ContentFolders)
+		if (!MountPointRootPath.IsEmpty())
 		{
-			// Parse the parent folder of .../Content from "Folder"
-			FString Path, Content, Root;
-			bool bSplitContent = Folder.Split(TEXT("/"), &Path, &Content, ESearchCase::CaseSensitive, ESearchDir::FromEnd);
-			bool bSplitParentPath = Path.Split(TEXT("/"), nullptr, &Root, ESearchCase::CaseSensitive, ESearchDir::FromEnd);
-			if (!bSplitContent || !bSplitParentPath)
+
+			FPackageMigrationImpl::FPackageMigrationImplContext MigrationImplContext;
+			MigrationImplContext.Options = Options;
+
+			// 1) Create the package migration state. It will create the mount point also if needed
+			UE::AssetTools::FPackageMigrationContext PackageMigrationContext(UE::AssetTools::FPackageMigrationContext::FScopedMountPoint(MoveTemp(MountPointRootPath), MoveTemp(DestinationFolder)));
+
+			PackageMigrationContext.CurrentStep =  UE::AssetTools::FPackageMigrationContext::EPackageMigrationStep::BeginMigration;
+			OnPackageMigration.Broadcast(PackageMigrationContext);
+
+			// 2) Set up the instancing context so that the packages are load in the destination mount point and setup the packages migration data
 			{
-				MigrateLog.Error(FText::Format(LOCTEXT("MigratePackages_NoMountPointFolder", "Unable to determine mount point for folder {0}"), FText::FromString(Folder)));
-				bAbort = true;
-				continue;
+				PackageMigrationContext.MigrationPackagesData.Reserve(PackageDataToMigrate->Num());
+
+				TArray<TPair<const ReportPackageData*, const FAssetData>> ExternalPackageDatas;
+				TMap<const FStringView, int32> ExistingPackageNameToMigrationDataIndex;
+
+				FPackageMigrationImpl::SetupPublicAssetPackagesMigrationData(PackageDataToMigrate
+					, ExternalPackageDatas
+					, ExistingPackageNameToMigrationDataIndex
+					, PackageMigrationContext
+					, MigrationImplContext
+					, MigrateLog);
+
+				PackageMigrationContext.CurrentStep =  UE::AssetTools::FPackageMigrationContext::EPackageMigrationStep::PostAssetMigrationPackageDataCreated;
+				OnPackageMigration.Broadcast(PackageMigrationContext);
+
+
+				FPackageMigrationImpl::SetupExternalAssetPackagesMigrationData(ExternalPackageDatas
+					, ExistingPackageNameToMigrationDataIndex
+					, PackageMigrationContext
+					, MigrationImplContext);
+
+				PackageMigrationContext.CurrentStep = UE::AssetTools::FPackageMigrationContext::EPackageMigrationStep::PostExternalMigrationPackageDataCreated;
+				OnPackageMigration.Broadcast(PackageMigrationContext);
+
 			}
 
-			// Determine this folder name to be a content root in the destination
-			FString DestContentRoot = FString(TEXT("/")) + Root;
-			FString DestContentFolder = FPaths::ConvertRelativePathToFull(Folder + TEXT("/"));
-			DestContentRootsToFolders.Add(FName(DestContentRoot), DestContentFolder);
+			PackageMigrationContext.CurrentStep = UE::AssetTools::FPackageMigrationContext::EPackageMigrationStep::InTheWayPackagesMoved;
+			OnPackageMigration.Broadcast(PackageMigrationContext);
+
+
+			// 3) Instanced load
+			TArray<TWeakObjectPtr<UPackage>> PackagesToClean;
+			PackagesToClean.Reserve(PackageMigrationContext.MigrationPackagesData.Num());
+			FPackageMigrationImpl::CreatePackagesAndSetupLinkers(PackageMigrationContext, MigrationImplContext, PackagesToClean);
+
+			PackageMigrationContext.CurrentStep = UE::AssetTools::FPackageMigrationContext::EPackageMigrationStep::InstancedPackagesCreated;
+			OnPackageMigration.Broadcast(PackageMigrationContext);
+
+			FPackageMigrationImpl::LoadInstancedPackages(PackageMigrationContext, MigrationImplContext);
+
+			// Collect the garbage created by the load here and it also help detecting potential issues in migration (which can cause crashes).
+			CollectGarbage(GARBAGE_COLLECTION_KEEPFLAGS);
+
+			PackageMigrationContext.CurrentStep = UE::AssetTools::FPackageMigrationContext::EPackageMigrationStep::InstancedPackagesLoaded;
+			OnPackageMigration.Broadcast(PackageMigrationContext);
+
+
+			// 4) Saves
+			FPackageMigrationImpl::SaveInstancedPackagesIntoDestination(PackageMigrationContext, MigrationImplContext);
+
+			PackageMigrationContext.CurrentStep = UE::AssetTools::FPackageMigrationContext::EPackageMigrationStep::InstancedPackagesSaved;
+			OnPackageMigration.Broadcast(PackageMigrationContext);
+
+
+			// 5) Remove the packages created for the migration
+			FPackageMigrationImpl::CleanInstancedPackages(PackagesToClean);
+
+			PackageMigrationContext.CurrentStep = UE::AssetTools::FPackageMigrationContext::EPackageMigrationStep::PostCleaningInstancedPackages;
+			OnPackageMigration.Broadcast(PackageMigrationContext);
+
+			FPackageMigrationImpl::RestoreInTheWayPackages(PackageMigrationContext);
+
+			PackageMigrationContext.CurrentStep = UE::AssetTools::FPackageMigrationContext::EPackageMigrationStep::InstancedPackagesSaved;
+			OnPackageMigration.Broadcast(PackageMigrationContext);
+
+			FPackageMigrationImpl::BuildAndNotifyLogWindow(PackageMigrationContext, MigrateLog, MigrationImplContext);
+
+			PackageMigrationContext.CurrentStep = UE::AssetTools::FPackageMigrationContext::EPackageMigrationStep::EndMigration;
+			OnPackageMigration.Broadcast(PackageMigrationContext);
 		}
 	}
-
-	if (bAbort)
+	else
 	{
-		MigrateLog.Notify();
-		return;
-	}
+		bool bAbort = false;
 
-	// Check if the root of any of the packages to migrate have no destination
-	TArray<FName> LostPackages;
-	TSet<FName> LostPackageRoots;
-	for (const FName& PackageName : AllPackageNamesToMove)
-	{
-		// Acquire the mount point for this package
-		FString Folder = TEXT("/") + FPackageName::GetPackageMountPoint(PackageName.ToString()).ToString();
+		// Associate each Content folder in the target Plugin hierarchy to a content root string in UFS
+		TMap<FName, FString> DestContentRootsToFolders;
 
-		// If this is /Game package, it doesn't need special handling and is simply bound for the destination, continue
-		if (Folder == TEXT("/Game"))
+		// Assets in /Game always map directly to the destination
+		DestContentRootsToFolders.Add(FName(TEXT("/Game")), DestinationFolder);
+
+		// If our destination is a project, it could have plugins...
+		if (bIsDestinationAProject)
 		{
-			continue;
+			// Find all "Content" folders under the destination ../Plugins directory
+			TArray<FString> ContentFolders;
+			IFileManager::Get().FindFilesRecursive(ContentFolders, *(DestinationFolder + TEXT("../Plugins/")), TEXT("Content"), false, true);
+
+			for (const FString& Folder : ContentFolders)
+			{
+				// Parse the parent folder of .../Content from "Folder"
+				FString Path, Content, Root;
+				bool bSplitContent = Folder.Split(TEXT("/"), &Path, &Content, ESearchCase::CaseSensitive, ESearchDir::FromEnd);
+				bool bSplitParentPath = Path.Split(TEXT("/"), nullptr, &Root, ESearchCase::CaseSensitive, ESearchDir::FromEnd);
+				if (!bSplitContent || !bSplitParentPath)
+				{
+					MigrateLog.Error(FText::Format(LOCTEXT("MigratePackages_NoMountPointFolder", "Unable to determine mount point for folder {0}"), FText::FromString(Folder)));
+					bAbort = true;
+					continue;
+				}
+
+				// Determine this folder name to be a content root in the destination
+				FString DestContentRoot = FString(TEXT("/")) + Root;
+				FString DestContentFolder = FPaths::ConvertRelativePathToFull(Folder + TEXT("/"));
+				DestContentRootsToFolders.Add(FName(DestContentRoot), DestContentFolder);
+			}
 		}
 
-		// Resolve the disk folder of this package's mount point so we compare directory names directly
-		//  We do this FileSystem to FileSystem compare instead of mount point (UFS) to content root (FileSystem)
-		//  The mount point name likely comes from the FriendlyName in the uplugin, or the uplugin basename
-		//  What's important here is that we succeed in finding _the same plugin_ between a copy/pasted project
-		if (!FPackageName::TryConvertLongPackageNameToFilename(Folder, Folder))
+		if (bAbort)
 		{
-			MigrateLog.Error(FText::Format(LOCTEXT("MigratePackages_NoContentFolder", "Unable to determine content folder for asset {0}"), FText::FromString(PackageName.ToString())));
-			bAbort = true;
-			continue;
+			MigrateLog.Notify();
+			return;
 		}
-		Folder.RemoveFromEnd(TEXT("/"));
 
-		// Parse the parent folder of .../Content from "Folder"
-		FString Path, Content, Root;
-		bool bSplitContent = Folder.Split(TEXT("/"), &Path, &Content, ESearchCase::CaseSensitive, ESearchDir::FromEnd);
-		bool bSplitParentPath = Path.Split(TEXT("/"), nullptr, &Root, ESearchCase::CaseSensitive, ESearchDir::FromEnd);
-		if (!bSplitContent || !bSplitParentPath)
-		{
-			MigrateLog.Error(FText::Format(LOCTEXT("MigratePackages_NoMountPointPackage", "Unable to determine mount point for package {0}"), FText::FromString(PackageName.ToString())));
-			bAbort = true;
-			continue;
+		// Check if the root of any of the packages to migrate have no destination
+		TArray<FName> LostPackages;
+		TSet<FName> LostPackageRoots;
+		for (const FName& PackageName : AllPackageNamesToMove)
+		{
+			// Acquire the mount point for this package
+			FString Folder = TEXT("/") + FPackageName::GetPackageMountPoint(PackageName.ToString()).ToString();
+
+			// If this is /Game package, it doesn't need special handling and is simply bound for the destination, continue
+			if (Folder == TEXT("/Game"))
+			{
+				continue;
+			}
+
+			// Resolve the disk folder of this package's mount point so we compare directory names directly
+			//  We do this FileSystem to FileSystem compare instead of mount point (UFS) to content root (FileSystem)
+			//  The mount point name likely comes from the FriendlyName in the uplugin, or the uplugin basename
+			//  What's important here is that we succeed in finding _the same plugin_ between a copy/pasted project
+			if (!FPackageName::TryConvertLongPackageNameToFilename(Folder, Folder))
+			{
+				MigrateLog.Error(FText::Format(LOCTEXT("MigratePackages_NoContentFolder", "Unable to determine content folder for asset {0}"), FText::FromString(PackageName.ToString())));
+				bAbort = true;
+				continue;
+			}
+			Folder.RemoveFromEnd(TEXT("/"));
+
+			// Parse the parent folder of .../Content from "Folder"
+			FString Path, Content, Root;
+			bool bSplitContent = Folder.Split(TEXT("/"), &Path, &Content, ESearchCase::CaseSensitive, ESearchDir::FromEnd);
+			bool bSplitParentPath = Path.Split(TEXT("/"), nullptr, &Root, ESearchCase::CaseSensitive, ESearchDir::FromEnd);
+			if (!bSplitContent || !bSplitParentPath)
+			{
+				MigrateLog.Error(FText::Format(LOCTEXT("MigratePackages_NoMountPointPackage", "Unable to determine mount point for package {0}"), FText::FromString(PackageName.ToString())));
+				bAbort = true;
+				continue;
+			}
+
+			// Check to see if the content root exists in the destination, otherwise it's "Lost"
+			FString SrcContentRoot = FString(TEXT("/")) + Root;
+			FName SrcContentRootName(SrcContentRoot);
+			if (!DestContentRootsToFolders.Find(SrcContentRootName))
+			{
+				LostPackages.Add(PackageName);
+				LostPackageRoots.Add(SrcContentRootName);
+			}
 		}
 
-		// Check to see if the content root exists in the destination, otherwise it's "Lost"
-		FString SrcContentRoot = FString(TEXT("/")) + Root;
-		FName SrcContentRootName(SrcContentRoot);
-		if (!DestContentRootsToFolders.Find(SrcContentRootName))
+		if (bAbort)
 		{
-			LostPackages.Add(PackageName);
-			LostPackageRoots.Add(SrcContentRootName);
+			MigrateLog.Notify();
+			return;
 		}
-	}
-
-	if (bAbort)
-	{
-		MigrateLog.Notify();
-		return;
-	}
 
-	// If some packages don't have a matching content root in the destination, prompt for desired behavior
-	if (!LostPackages.IsEmpty())
-	{
-		if (!FApp::IsUnattended() && Options.bPrompt && Options.OrphanFolder.IsEmpty())
+		// If some packages don't have a matching content root in the destination, prompt for desired behavior
+		if (!LostPackages.IsEmpty())
 		{
-			FString LostPackageRootsString;
-			for (const FName& PackageRoot : LostPackageRoots)
+			if (!FApp::IsUnattended() && Options.bPrompt && Options.OrphanFolder.IsEmpty())
 			{
-				LostPackageRootsString += FString(TEXT("\n\t")) + PackageRoot.ToString();
-			}
+				FString LostPackageRootsString;
+				for (const FName& PackageRoot : LostPackageRoots)
+				{
+					LostPackageRootsString += FString(TEXT("\n\t")) + PackageRoot.ToString();
+				}
 
-			// Prompt to consolidate to a migration folder
-			FText Prompt = FText::Format(LOCTEXT("MigratePackages_ConsolidateToTemp", "Some selected assets don't have a corresponding content root in the destination.{0}\n\nWould you like to save a copy of all selected assets into a folder with consolidated references? If you select No then assets in the above roots will not be migrated."), FText::FromString(LostPackageRootsString));
-			switch (FMessageDialog::Open(EAppMsgType::YesNoCancel, Prompt))
+				// Prompt to consolidate to a migration folder
+				FText Prompt = FText::Format(LOCTEXT("MigratePackages_ConsolidateToTemp", "Some selected assets don't have a corresponding content root in the destination.{0}\n\nWould you like to save a copy of all selected assets into a folder with consolidated references? If you select No then assets in the above roots will not be migrated."), FText::FromString(LostPackageRootsString));
+				switch (FMessageDialog::Open(EAppMsgType::YesNoCancel, Prompt))
+				{
+				case EAppReturnType::Yes:
+					// No op
+					break;
+				case EAppReturnType::No:
+					LostPackages.Reset();
+					break;
+				case EAppReturnType::Cancel:
+					return;
+				}
+			}
+			else if (Options.OrphanFolder.IsEmpty())
 			{
-			case EAppReturnType::Yes:
-				// No op
-				break;
-			case EAppReturnType::No:
 				LostPackages.Reset();
-				break;
-			case EAppReturnType::Cancel:
-				return;
 			}
 		}
-		else if (Options.OrphanFolder.IsEmpty())
-		{
-			LostPackages.Reset();
-		}
-	}
 
-	// This will be used to tidy up the temp folder after we copy the migrated assets
-	FString SrcDiskFolderFilename;
+		// This will be used to tidy up the temp folder after we copy the migrated assets
+		FString SrcDiskFolderFilename;
 
-	// Fixing up references requires resaving packages to a temporary location
-	if (!LostPackages.IsEmpty())
-	{
-		// Query the user for a folder to migrate assets into. This folder will exist temporary in this project, and is the destination for migration in the target project
-		FString FolderName = Options.OrphanFolder.IsEmpty()? TEXT("Migrated") : Options.OrphanFolder;
-		if (!FApp::IsUnattended() && Options.bPrompt)
-		{
-			bool bIsOkButtonEnabled = true;
-			TSharedRef<SEditableTextBox> EditableTextBox = SNew(SEditableTextBox)
-				.Text(FText::FromString(FolderName))
-				.OnVerifyTextChanged_Lambda([&bIsOkButtonEnabled](const FText& InNewText, FText& OutErrorMessage) -> bool
-					{
-						if (InNewText.ToString().Contains(TEXT("/")))
+		// Fixing up references requires resaving packages to a temporary location
+		if (!LostPackages.IsEmpty())
+		{
+			// Query the user for a folder to migrate assets into. This folder will exist temporary in this project, and is the destination for migration in the target project
+			FString FolderName = Options.OrphanFolder.IsEmpty() ? TEXT("Migrated") : Options.OrphanFolder;
+			if (!FApp::IsUnattended() && Options.bPrompt)
+			{
+				bool bIsOkButtonEnabled = true;
+				TSharedRef<SEditableTextBox> EditableTextBox = SNew(SEditableTextBox)
+					.Text(FText::FromString(FolderName))
+					.OnVerifyTextChanged_Lambda([&bIsOkButtonEnabled](const FText& InNewText, FText& OutErrorMessage) -> bool
 						{
-							OutErrorMessage = LOCTEXT("Migrated_CannotContainSlashes", "Cannot use a slash in a folder name.");
+							if (InNewText.ToString().Contains(TEXT("/")))
+							{
+								OutErrorMessage = LOCTEXT("Migrated_CannotContainSlashes", "Cannot use a slash in a folder name.");
 
-							// Disable Ok if the string is invalid
-							bIsOkButtonEnabled = false;
-							return false;
-						}
+								// Disable Ok if the string is invalid
+								bIsOkButtonEnabled = false;
+								return false;
+							}
+
+							// Enable Ok if the string is valid
+							bIsOkButtonEnabled = true;
+							return true;
+						})
+					.OnTextCommitted_Lambda([&FolderName](const FText& NewValue, ETextCommit::Type)
+						{
+							// Set the result if they modified the text
+							FolderName = NewValue.ToString();
+						});
 
-						// Enable Ok if the string is valid
-						bIsOkButtonEnabled = true;
-						return true;
-					})
-				.OnTextCommitted_Lambda([&FolderName](const FText& NewValue, ETextCommit::Type)
+				// Set the result if they just click Ok
+				SGenericDialogWidget::FArguments FolderDialogArguments;
+				FolderDialogArguments.OnOkPressed_Lambda([&EditableTextBox, &FolderName]()
 					{
-						// Set the result if they modified the text
-						FolderName = NewValue.ToString();
+						FolderName = EditableTextBox->GetText().ToString();
 					});
 
-			// Set the result if they just click Ok
-			SGenericDialogWidget::FArguments FolderDialogArguments;
-			FolderDialogArguments.OnOkPressed_Lambda([&EditableTextBox, &FolderName]()
+				// Present the Dialog
+				SGenericDialogWidget::OpenDialog(LOCTEXT("MigratePackages_FolderName", "Folder for Migrated Assets"),
+					SNew(SVerticalBox)
+					+ SVerticalBox::Slot()
+					.AutoHeight()
+					.Padding(5.0f)
+					[
+						SNew(STextBlock)
+						.Text(LOCTEXT("MigratePackages_SpecifyConsolidateFolder", "Please specify a new folder name to consolidate the assets into."))
+					]
+					+ SVerticalBox::Slot()
+					.Padding(5.0f)
+					[
+						SNew(SSpacer)
+					]
+					+ SVerticalBox::Slot()
+					.AutoHeight()
+					.Padding(5.0f)
+					[
+						EditableTextBox
+					],
+					FolderDialogArguments, true);
+
+				// Sanity the user input
+				if (FolderName.IsEmpty())
 				{
-					FolderName = EditableTextBox->GetText().ToString();
-				});
-
-			// Present the Dialog
-			SGenericDialogWidget::OpenDialog(LOCTEXT("MigratePackages_FolderName", "Folder for Migrated Assets"),
-				SNew(SVerticalBox)
-				+ SVerticalBox::Slot()
-				.AutoHeight()
-				.Padding(5.0f)
-				[
-					SNew(STextBlock)
-					.Text(LOCTEXT("MigratePackages_SpecifyConsolidateFolder", "Please specify a new folder name to consolidate the assets into."))
-				]
-			+ SVerticalBox::Slot()
-				.Padding(5.0f)
-				[
-					SNew(SSpacer)
-				]
-			+ SVerticalBox::Slot()
-				.AutoHeight()
-				.Padding(5.0f)
-				[
-					EditableTextBox
-				],
-				FolderDialogArguments, true);
-
-			// Sanity the user input
-			if (FolderName.IsEmpty())
-			{
-				return;
-			}
-		}
-		// Remove forbidden characters
-		FolderName = FolderName.Replace(TEXT("/"), TEXT(""));
-
-		// Verify that we don't have any assets that exist where we want to perform our consolidation
-		//  We could do the same check at the destination, but the package copy code will happily overwrite,
-		//  and it seems reasonable someone would want to Migrate several times if snapshotting changes from another project.
-		//  So for now we don't necessitate that the destination content folder be missing as well.
-		FString SrcUfsFolderName = FolderName;
-		SrcUfsFolderName.InsertAt(0, TEXT("/Game/"));
-		SrcDiskFolderFilename = FPackageName::LongPackageNameToFilename(SrcUfsFolderName, TEXT(""));
-		if (IFileManager::Get().DirectoryExists(*SrcDiskFolderFilename))
-		{
-			const FText Message = FText::Format(LOCTEXT("MigratePackages_InvalidMigrateFolder", "{0} exists on disk in the source project, and cannot be used to consolidate assets."), FText::FromString(SrcDiskFolderFilename));
-			if (!FApp::IsUnattended() && Options.bPrompt)
-			{
-				EAppReturnType::Type Response = FMessageDialog::Open(EAppMsgType::Ok, Message);
-			}
-			else
-			{
-				UE_LOG(LogAssetTools, Error, TEXT("%s"), *Message.ToString());
+					return;
+				}
 			}
-			return;
-		}
-
-		// Resolve the packages to migrate to assets
-		TArray<UObject*> SrcObjects;
-		{
-			FScopedSlowTask SlowTask(AllPackageNamesToMove.Num(), LOCTEXT("MigratePackages_Loading", "Loading Packages..."));
-			SlowTask.MakeDialog();
-			for (const FName& SrcPackage : AllPackageNamesToMove)
+			// Remove forbidden characters
+			FolderName = FolderName.Replace(TEXT("/"), TEXT(""));
+
+			// Verify that we don't have any assets that exist where we want to perform our consolidation
+			//  We could do the same check at the destination, but the package copy code will happily overwrite,
+			//  and it seems reasonable someone would want to Migrate several times if snapshotting changes from another project.
+			//  So for now we don't necessitate that the destination content folder be missing as well.
+			FString SrcUfsFolderName = FolderName;
+			SrcUfsFolderName.InsertAt(0, TEXT("/Game/"));
+			SrcDiskFolderFilename = FPackageName::LongPackageNameToFilename(SrcUfsFolderName, TEXT(""));
+			if (IFileManager::Get().DirectoryExists(*SrcDiskFolderFilename))
 			{
-				SlowTask.EnterProgressFrame();
-				UPackage* LoadedPackage = UPackageTools::LoadPackage(SrcPackage.ToString());
-				if (!LoadedPackage)
-				{
-					MigrateLog.Error(FText::Format(LOCTEXT("MigratePackages_FailedToLoadPackage", "Failed to load package {0}"), FText::FromString(SrcPackage.ToString())));
-					bAbort = true;
-					continue;
-				}
-
-				UObject* Object = LoadedPackage->FindAssetInPackage();
-				if (Object)
+				const FText Message = FText::Format(LOCTEXT("MigratePackages_InvalidMigrateFolder", "{0} exists on disk in the source project, and cannot be used to consolidate assets."), FText::FromString(SrcDiskFolderFilename));
+				if (!FApp::IsUnattended() && Options.bPrompt)
 				{
-					SrcObjects.Add(Object);
+					EAppReturnType::Type Response = FMessageDialog::Open(EAppMsgType::Ok, Message);
 				}
 				else
 				{
-					MigrateLog.Warning(FText::Format(LOCTEXT("MigratePackages_PackageHasNoAsset", "Package {0} has no asset in it"), FText::FromString(SrcPackage.ToString())));
+					UE_LOG(LogAssetTools, Error, TEXT("%s"), *Message.ToString());
 				}
+				return;
+			}
 
-				UBlueprint* Blueprint = Cast<UBlueprint>(Object);
-				if (Blueprint && (!Blueprint->SkeletonGeneratedClass || !Blueprint->GeneratedClass))
+			// Resolve the packages to migrate to assets
+			TArray<UObject*> SrcObjects;
+			{
+				FScopedSlowTask SlowTask(AllPackageNamesToMove.Num(), LOCTEXT("MigratePackages_Loading", "Loading Packages..."));
+				SlowTask.MakeDialog();
+				for (const FName& SrcPackage : AllPackageNamesToMove)
 				{
-					MigrateLog.Error(FText::Format(LOCTEXT("MigratePackages_InvalidBluepriint", "Blueprint is invalid in package {0}"), FText::FromString(SrcPackage.ToString())));
-					bAbort = true;
-					continue;
-				}
-			}
-		}
+					SlowTask.EnterProgressFrame();
+					UPackage* LoadedPackage = UPackageTools::LoadPackage(SrcPackage.ToString());
+					if (!LoadedPackage)
+					{
+						MigrateLog.Error(FText::Format(LOCTEXT("MigratePackages_FailedToLoadPackage", "Failed to load package {0}"), FText::FromString(SrcPackage.ToString())));
+						bAbort = true;
+						continue;
+					}
 
-		if (bAbort)
-		{
-			MigrateLog.Notify();
-			return;
-		}
+					UObject* Object = LoadedPackage->FindAssetInPackage();
+					if (Object)
+					{
+						SrcObjects.Add(Object);
+					}
+					else
+					{
+						MigrateLog.Warning(FText::Format(LOCTEXT("MigratePackages_PackageHasNoAsset", "Package {0} has no asset in it"), FText::FromString(SrcPackage.ToString())));
+					}
 
-		// To handle complex references and assets in different Plugins, we must first duplicate to temp packages
-		TArray<UObject*> TempObjects;
-		TMap<UObject*, UObject*> ReplacementMap;
-		{
-			FScopedSlowTask SlowTask(SrcObjects.Num(), LOCTEXT("MigratePackages_Duplicate", "Duplicating Assets..."));
-			SlowTask.MakeDialog();
+					UBlueprint* Blueprint = Cast<UBlueprint>(Object);
+					if (Blueprint && (!Blueprint->SkeletonGeneratedClass || !Blueprint->GeneratedClass))
+					{
+						MigrateLog.Error(FText::Format(LOCTEXT("MigratePackages_InvalidBluepriint", "Blueprint is invalid in package {0}"), FText::FromString(SrcPackage.ToString())));
+						bAbort = true;
+						continue;
+					}
+				}
+			}
 
-			TSet<UPackage*> PackagesUserRefusedToFullyLoad;
-			for (int i=0; i<SrcObjects.Num(); ++i)
+			if (bAbort)
 			{
-				SlowTask.EnterProgressFrame();
+				MigrateLog.Notify();
+				return;
+			}
 
-				UObject* Object = SrcObjects[i];
-				FString PackageName = Object->GetPackage()->GetName();
+			// To handle complex references and assets in different Plugins, we must first duplicate to temp packages
+			TArray<UObject*> TempObjects;
+			TMap<UObject*, UObject*> ReplacementMap;
+			{
+				FScopedSlowTask SlowTask(SrcObjects.Num(), LOCTEXT("MigratePackages_Duplicate", "Duplicating Assets..."));
+				SlowTask.MakeDialog();
 
-				FString Path, Root;
-				PackageName.RemoveFromStart(TEXT("/"));
-				bool bSplitRoot = PackageName.Split(TEXT("/"), &Root, &Path);
-				if (!bSplitRoot)
+				TSet<UPackage*> PackagesUserRefusedToFullyLoad;
+				for (int i=0; i<SrcObjects.Num(); ++i)
 				{
-					MigrateLog.Error(FText::Format(LOCTEXT("MigratePackages_NoMountPointPackage", "Unable to determine mount point for package {0}"), FText::FromString(PackageName)));
-					continue;
-				}
+					SlowTask.EnterProgressFrame();
 
-				ObjectTools::FPackageGroupName PackageGroupName;
-				PackageGroupName.ObjectName = Object->GetName();
-				PackageGroupName.PackageName = SrcUfsFolderName + TEXT("/") + Path;
-				FString GroupName = Object->GetFullGroupName(/*bStartWithOuter =*/true);
-				if (GroupName != TEXT("None"))
-				{
-					PackageGroupName.GroupName = GroupName;
-				}
+					UObject* Object = SrcObjects[i];
+					FString PackageName = Object->GetPackage()->GetName();
 
-				UObject* NewObject = ObjectTools::DuplicateSingleObject(Object, PackageGroupName, PackagesUserRefusedToFullyLoad);
-				if (NewObject)
-				{
-					TempObjects.Add(NewObject);
-					ReplacementMap.Add(SrcObjects[i], TempObjects[i]);
+					FString Path, Root;
+					PackageName.RemoveFromStart(TEXT("/"));
+					bool bSplitRoot = PackageName.Split(TEXT("/"), &Root, &Path);
+					if (!bSplitRoot)
+					{
+						MigrateLog.Error(FText::Format(LOCTEXT("MigratePackages_NoMountPointPackage", "Unable to determine mount point for package {0}"), FText::FromString(PackageName)));
+						continue;
+					}
+
+					ObjectTools::FPackageGroupName PackageGroupName;
+					PackageGroupName.ObjectName = Object->GetName();
+					PackageGroupName.PackageName = SrcUfsFolderName + TEXT("/") + Path;
+					FString GroupName = Object->GetFullGroupName(/*bStartWithOuter =*/true);
+					if (GroupName != TEXT("None"))
+					{
+						PackageGroupName.GroupName = GroupName;
+					}
+
+					UObject* NewObject = ObjectTools::DuplicateSingleObject(Object, PackageGroupName, PackagesUserRefusedToFullyLoad);
+					if (NewObject)
+					{
+						TempObjects.Add(NewObject);
+						ReplacementMap.Add(SrcObjects[i], TempObjects[i]);
+					}
 				}
 			}
-		}
-
-		// Update references between TempObjects (to reference each other)
-		{
-			FScopedSlowTask SlowTask(TempObjects.Num(), LOCTEXT("MigratePackages_ReplaceReferences", "Replacing References..."));
-			SlowTask.MakeDialog();
 
-			for (int i=0; i<TempObjects.Num(); ++i)
+			// Update references between TempObjects (to reference each other)
 			{
-				SlowTask.EnterProgressFrame();
-				UObject* TempObject = TempObjects[i];
+				FScopedSlowTask SlowTask(TempObjects.Num(), LOCTEXT("MigratePackages_ReplaceReferences", "Replacing References..."));
+				SlowTask.MakeDialog();
 
-				FArchiveReplaceObjectRef<UObject> ReplaceAr(TempObject, ReplacementMap, EArchiveReplaceObjectFlags::IgnoreOuterRef | EArchiveReplaceObjectFlags::IgnoreArchetypeRef);
-			}
-		}
+				for (int i=0; i<TempObjects.Num(); ++i)
+				{
+					SlowTask.EnterProgressFrame();
+					UObject* TempObject = TempObjects[i];
 
-		// Save fixed up packages to the migrated folder, and update the set of files to copy to be those migrated packages
-		{
-			FScopedSlowTask SlowTask(TempObjects.Num(), LOCTEXT("MigratePackages_ReplaceReferences", "Replacing References..."));
-			SlowTask.MakeDialog();
+					FArchiveReplaceObjectRef<UObject> ReplaceAr(TempObject, ReplacementMap, EArchiveReplaceObjectFlags::IgnoreOuterRef | EArchiveReplaceObjectFlags::IgnoreArchetypeRef);
+				}
+			}
 
-			TSet<FName> NewPackageNamesToMove;
-			for (int i=0; i<TempObjects.Num(); ++i)
+			// Save fixed up packages to the migrated folder, and update the set of files to copy to be those migrated packages
 			{
-				SlowTask.EnterProgressFrame();
-				UObject* TempObject = TempObjects[i];
+				FScopedSlowTask SlowTask(TempObjects.Num(), LOCTEXT("MigratePackages_ReplaceReferences", "Replacing References..."));
+				SlowTask.MakeDialog();
 
-				// Calculate the file path to the new, migrated package
-				FString const TempPackageName = TempObject->GetPackage()->GetName();
-				FString const TempPackageFilename = FPackageName::LongPackageNameToFilename(TempPackageName, FPackageName::GetAssetPackageExtension());
-
-				// Save it
-				FSavePackageArgs SaveArgs;
-				GEditor->Save(TempObject->GetPackage(), /*InAsset=*/nullptr, *TempPackageFilename, SaveArgs);
+				TSet<FName> NewPackageNamesToMove;
+				for (int i=0; i<TempObjects.Num(); ++i)
+				{
+					SlowTask.EnterProgressFrame();
+					UObject* TempObject = TempObjects[i];
 
-				NewPackageNamesToMove.Add(FName(TempPackageName));
-			}
+					// Calculate the file path to the new, migrated package
+					FString const TempPackageName = TempObject->GetPackage()->GetName();
+					FString const TempPackageFilename = FPackageName::LongPackageNameToFilename(TempPackageName, FPackageName::GetAssetPackageExtension());
 
-			AllPackageNamesToMove = NewPackageNamesToMove;
-		}
-	}
+					// Save it
+					FSavePackageArgs SaveArgs;
+					GEditor->Save(TempObject->GetPackage(), /*InAsset=*/nullptr, *TempPackageFilename, SaveArgs);
 
-	bool bUserCanceled = false;
+					NewPackageNamesToMove.Add(FName(TempPackageName));
+				}
 
-	EAppReturnType::Type LastResponse = EAppReturnType::Yes;
-	if (!Options.bPrompt)
-	{
-		switch (Options.AssetConflict) {
-		case EAssetMigrationConflict::Overwrite:
-			LastResponse = EAppReturnType::YesAll;
-			break;
-		case EAssetMigrationConflict::Skip:
-			LastResponse = EAppReturnType::NoAll;
-			break;
-		case EAssetMigrationConflict::Cancel:
-			LastResponse = EAppReturnType::Cancel;
-			break;
+				AllPackageNamesToMove = NewPackageNamesToMove;
+			}
 		}
-	}
-	TArray<FString> SuccessfullyCopiedFiles;
-	TArray<FString> SuccessfullyCopiedPackages;
-	FString CopyErrors;
 
-	{
-		FScopedSlowTask SlowTask(AllPackageNamesToMove.Num(), LOCTEXT("MigratePackages_CopyingFiles", "Copying Files..."));
-		SlowTask.MakeDialog();
+		bool bUserCanceled = false;
 
-		for ( const FName& PackageNameToMove : AllPackageNamesToMove )
+		EAppReturnType::Type LastResponse = EAppReturnType::Yes;
+		if (!Options.bPrompt)
 		{
-			SlowTask.EnterProgressFrame();
-
-			const FString& PackageName = PackageNameToMove.ToString();
-			FString SrcFilename;
-			
-			if (!FPackageName::DoesPackageExist(PackageName, &SrcFilename))
-			{
-				const FText ErrorMessage = FText::Format(LOCTEXT("MigratePackages_PackageMissing", "{0} does not exist on disk."), FText::FromString(PackageName));
-				UE_LOG(LogAssetTools, Warning, TEXT("%s"), *ErrorMessage.ToString());
-				CopyErrors += ErrorMessage.ToString() + LINE_TERMINATOR;
-			}
-			else if (SrcFilename.Contains(FPaths::EngineContentDir()))
+			switch (Options.AssetConflict)
 			{
-				const FString LeafName = SrcFilename.Replace(*FPaths::EngineContentDir(), TEXT("Engine/"));
-				CopyErrors += FText::Format(LOCTEXT("MigratePackages_EngineContent", "Unable to migrate Engine asset {0}. Engine assets cannot be migrated."), FText::FromString(LeafName)).ToString() + LINE_TERMINATOR;
+				case EAssetMigrationConflict::Overwrite:
+					LastResponse = EAppReturnType::YesAll;
+					break;
+				case EAssetMigrationConflict::Skip:
+					LastResponse = EAppReturnType::NoAll;
+					break;
+				case EAssetMigrationConflict::Cancel:
+					LastResponse = EAppReturnType::Cancel;
+					break;
 			}
-			else
-			{
-				bool bFileOKToCopy = true;
+		}
+		TArray<FString> SuccessfullyCopiedFiles;
+		TArray<FString> SuccessfullyCopiedPackages;
+		FString CopyErrors;
+
+		{
+			FScopedSlowTask SlowTask(AllPackageNamesToMove.Num(), LOCTEXT("MigratePackages_CopyingFiles", "Copying Files..."));
+			SlowTask.MakeDialog();
 
-				FString Path = PackageNameToMove.ToString();
-				FString PackageRoot;
-				Path.RemoveFromStart(TEXT("/"));
-				Path.Split("/", &PackageRoot, &Path, ESearchCase::IgnoreCase, ESearchDir::FromStart);
-				PackageRoot = TEXT("/") + PackageRoot;
+			for ( const FName& PackageNameToMove : AllPackageNamesToMove )
+			{
+				SlowTask.EnterProgressFrame();
 
-				FString DestFilename;
-				TMap<FName, FString>::ValueType* DestRootFolder = DestContentRootsToFolders.Find(FName(PackageRoot));
-				if (ensure(DestRootFolder))
+				const FString& PackageName = PackageNameToMove.ToString();
+				FString SrcFilename;
+			
+				if (!FPackageName::DoesPackageExist(PackageName, &SrcFilename))
 				{
-					DestFilename = *DestRootFolder;
+					const FText ErrorMessage = FText::Format(LOCTEXT("MigratePackages_PackageMissing", "{0} does not exist on disk."), FText::FromString(PackageName));
+					UE_LOG(LogAssetTools, Warning, TEXT("%s"), *ErrorMessage.ToString());
+					CopyErrors += ErrorMessage.ToString() + LINE_TERMINATOR;
 				}
-				else
+				else if (SrcFilename.Contains(FPaths::EngineContentDir()))
 				{
-					continue;
+					const FString LeafName = SrcFilename.Replace(*FPaths::EngineContentDir(), TEXT("Engine/"));
+					CopyErrors += FText::Format(LOCTEXT("MigratePackages_EngineContent", "Unable to migrate Engine asset {0}. Engine assets cannot be migrated."), FText::FromString(LeafName)).ToString() + LINE_TERMINATOR;
 				}
-
-				FString SubFolder;
-				if ( SrcFilename.Split( TEXT("/Content/"), nullptr, &SubFolder ) )
+				else
 				{
-					DestFilename += *SubFolder;
+					bool bFileOKToCopy = true;
+
+					FString Path = PackageNameToMove.ToString();
+					FString PackageRoot;
+					Path.RemoveFromStart(TEXT("/"));
+					Path.Split("/", &PackageRoot, &Path, ESearchCase::IgnoreCase, ESearchDir::FromStart);
+					PackageRoot = TEXT("/") + PackageRoot;
 
-					if ( IFileManager::Get().FileSize(*DestFilename) > 0 )
+					FString DestFilename;
+					TMap<FName, FString>::ValueType* DestRootFolder = DestContentRootsToFolders.Find(FName(PackageRoot));
+					if (ensure(DestRootFolder))
 					{
-						// The destination file already exists! Ask the user what to do.
-						EAppReturnType::Type Response;
-						if (FApp::IsUnattended() || !Options.bPrompt || LastResponse == EAppReturnType::YesAll || LastResponse == EAppReturnType::NoAll)
-						{
-							Response = LastResponse;
-						}
-						else
-						{
-							const FText Message = FText::Format( LOCTEXT("MigratePackages_AlreadyExists", "An asset already exists at location {0} would you like to overwrite it?"), FText::FromString(DestFilename) );
-							Response = FMessageDialog::Open( EAppMsgType::YesNoYesAllNoAllCancel, Message );
-							LastResponse = Response;
-						}
+						DestFilename = *DestRootFolder;
+					}
+					else
+					{
+						continue;
+					}
 
-						if (Response == EAppReturnType::Cancel)
-						{
-							// The user chose to cancel mid-operation. Break out.
-							bUserCanceled = true;
-							break;
-						}
+					FString SubFolder;
+					if ( SrcFilename.Split( TEXT("/Content/"), nullptr, &SubFolder ) )
+					{
+						DestFilename += *SubFolder;
 
-						const bool bWantOverwrite = Response == EAppReturnType::Yes || Response == EAppReturnType::YesAll;
-						if( !bWantOverwrite )
+						if ( IFileManager::Get().FileSize(*DestFilename) > 0 )
 						{
-							// User chose not to replace the package
-							bFileOKToCopy = false;
+							// The destination file already exists! Ask the user what to do.
+							EAppReturnType::Type Response;
+							if (FApp::IsUnattended() || !Options.bPrompt || LastResponse == EAppReturnType::YesAll || LastResponse == EAppReturnType::NoAll)
+							{
+								Response = LastResponse;
+							}
+							else
+							{
+								const FText Message = FText::Format( LOCTEXT("MigratePackages_AlreadyExists", "An asset already exists at location {0} would you like to overwrite it?"), FText::FromString(DestFilename) );
+								Response = FMessageDialog::Open( EAppMsgType::YesNoYesAllNoAllCancel, Message );
+								LastResponse = Response;
+							}
+
+							if (Response == EAppReturnType::Cancel)
+							{
+								// The user chose to cancel mid-operation. Break out.
+								bUserCanceled = true;
+								break;
+							}
+
+							const bool bWantOverwrite = Response == EAppReturnType::Yes || Response == EAppReturnType::YesAll;
+							if( !bWantOverwrite )
+							{
+								// User chose not to replace the package
+								bFileOKToCopy = false;
+							}
 						}
 					}
-				}
-				else
-				{
-					// Couldn't find Content folder in source path
-					bFileOKToCopy = false;
-				}
-
-				if ( bFileOKToCopy )
-				{
-					if ( IFileManager::Get().Copy(*DestFilename, *SrcFilename) == COPY_OK )
+					else
 					{
-						SuccessfullyCopiedPackages.Add(PackageName);
-						SuccessfullyCopiedFiles.Add(DestFilename);
+						// Couldn't find Content folder in source path
+						bFileOKToCopy = false;
 					}
-					else
+
+					if ( bFileOKToCopy )
 					{
-						UE_LOG(LogAssetTools, Warning, TEXT("Failed to copy %s to %s while migrating assets"), *SrcFilename, *DestFilename);
-						CopyErrors += SrcFilename + LINE_TERMINATOR;
+						if ( IFileManager::Get().Copy(*DestFilename, *SrcFilename) == COPY_OK )
+						{
+							SuccessfullyCopiedPackages.Add(PackageName);
+							SuccessfullyCopiedFiles.Add(DestFilename);
+						}
+						else
+						{
+							UE_LOG(LogAssetTools, Warning, TEXT("Failed to copy %s to %s while migrating assets"), *SrcFilename, *DestFilename);
+							CopyErrors += SrcFilename + LINE_TERMINATOR;
+						}
 					}
 				}
 			}
 		}
-	}
 
-	// If we are consolidating lost packages, we are copying temporary packages, so clean them up.
-	if (!LostPackages.IsEmpty())
-	{
-		FScopedSlowTask SlowTask(AllPackageNamesToMove.Num(), LOCTEXT("MigratePackages_CleaningUp", "Cleaning Up..."));
-		SlowTask.MakeDialog();
-
-		TArray<UObject*> ObjectsToDelete;
-		for (const FName& PackageNameToMove : AllPackageNamesToMove)
+		// If we are consolidating lost packages, we are copying temporary packages, so clean them up.
+		if (!LostPackages.IsEmpty())
 		{
-			SlowTask.EnterProgressFrame();
+			FScopedSlowTask SlowTask(AllPackageNamesToMove.Num(), LOCTEXT("MigratePackages_CleaningUp", "Cleaning Up..."));
+			SlowTask.MakeDialog();
 
-			UPackage* Package = UPackageTools::LoadPackage(PackageNameToMove.ToString());
-			if (Package)
+			TArray<UObject*> ObjectsToDelete;
+			for (const FName& PackageNameToMove : AllPackageNamesToMove)
 			{
-				ObjectsToDelete.Add(Package);
+				SlowTask.EnterProgressFrame();
+
+				UPackage* Package = UPackageTools::LoadPackage(PackageNameToMove.ToString());
+				if (Package)
+				{
+					ObjectsToDelete.Add(Package);
+				}
 			}
-		}
 
-		ObjectTools::ForceDeleteObjects(ObjectsToDelete, /*bShowConfirmation=*/false);
+			ObjectTools::ForceDeleteObjects(ObjectsToDelete, /*bShowConfirmation=*/false);
 
-		if (!IFileManager::Get().DeleteDirectory(*SrcDiskFolderFilename, /*RequireExists =*/ false, /*Tree =*/ true))
-		{
-			UE_LOG(LogAssetTools, Warning, TEXT("Failed to delete temporary directory %s while migrating assets"), *SrcDiskFolderFilename);
-			CopyErrors += SrcDiskFolderFilename + LINE_TERMINATOR;
+			if (!IFileManager::Get().DeleteDirectory(*SrcDiskFolderFilename, /*RequireExists =*/ false, /*Tree =*/ true))
+			{
+				UE_LOG(LogAssetTools, Warning, TEXT("Failed to delete temporary directory %s while migrating assets"), *SrcDiskFolderFilename);
+				CopyErrors += SrcDiskFolderFilename + LINE_TERMINATOR;
+			}
 		}
-	}
 
-	FString SourceControlErrors;
+		FString SourceControlErrors;
 
-	if ( !bUserCanceled && SuccessfullyCopiedFiles.Num() > 0 )
-	{
-		// attempt to add files to source control (this can quite easily fail, but if it works it is very useful)
-		if(GetDefault<UEditorLoadingSavingSettings>()->bSCCAutoAddNewFiles)
+		if ( !bUserCanceled && SuccessfullyCopiedFiles.Num() > 0 )
 		{
-			if(ISourceControlModule::Get().IsEnabled())
+			// attempt to add files to source control (this can quite easily fail, but if it works it is very useful)
+			if(GetDefault<UEditorLoadingSavingSettings>()->bSCCAutoAddNewFiles)
 			{
-				ISourceControlProvider& SourceControlProvider = ISourceControlModule::Get().GetProvider();
-				if(SourceControlProvider.Execute(ISourceControlOperation::Create<FMarkForAdd>(), SuccessfullyCopiedFiles) == ECommandResult::Failed)
+				if(ISourceControlModule::Get().IsEnabled())
 				{
-					FScopedSlowTask SlowTask(SuccessfullyCopiedFiles.Num(), LOCTEXT("MigratePackages_AddToSourceControl", "Adding To Source Control..."));
-					SlowTask.MakeDialog();
-
-					for(auto FileIt(SuccessfullyCopiedFiles.CreateConstIterator()); FileIt; FileIt++)
+					ISourceControlProvider& SourceControlProvider = ISourceControlModule::Get().GetProvider();
+					if(SourceControlProvider.Execute(ISourceControlOperation::Create<FMarkForAdd>(), SuccessfullyCopiedFiles) == ECommandResult::Failed)
 					{
-						SlowTask.EnterProgressFrame();
-						if(!SourceControlProvider.GetState(*FileIt, EStateCacheUsage::Use)->IsAdded())
+						FScopedSlowTask SlowTask(SuccessfullyCopiedFiles.Num(), LOCTEXT("MigratePackages_AddToSourceControl", "Adding To Source Control..."));
+						SlowTask.MakeDialog();
+
+						for(auto FileIt(SuccessfullyCopiedFiles.CreateConstIterator()); FileIt; FileIt++)
 						{
-							SourceControlErrors += FText::Format(LOCTEXT("MigratePackages_SourceControlError", "{0} could not be added to source control"), FText::FromString(*FileIt)).ToString();
-							SourceControlErrors += LINE_TERMINATOR;
+							SlowTask.EnterProgressFrame();
+							if(!SourceControlProvider.GetState(*FileIt, EStateCacheUsage::Use)->IsAdded())
+							{
+								SourceControlErrors += FText::Format(LOCTEXT("MigratePackages_SourceControlError", "{0} could not be added to source control"), FText::FromString(*FileIt)).ToString();
+								SourceControlErrors += LINE_TERMINATOR;
+							}
 						}
 					}
 				}
 			}
 		}
-	}
 
-	FText LogMessage = FText::FromString(TEXT("Content migration completed successfully!"));
-	EMessageSeverity::Type Severity = EMessageSeverity::Info;
-	if ( CopyErrors.Len() > 0 || SourceControlErrors.Len() > 0 )
-	{
-		FString ErrorMessage;
-		Severity = EMessageSeverity::Error;
-		if( CopyErrors.Len() > 0 )
+		FText LogMessage = FText::FromString(TEXT("Content migration completed successfully!"));
+		EMessageSeverity::Type Severity = EMessageSeverity::Info;
+		if ( CopyErrors.Len() > 0 || SourceControlErrors.Len() > 0 )
 		{
-			MigrateLog.NewPage( LOCTEXT("MigratePackages_CopyErrorsPage", "Copy Errors") );
-			MigrateLog.Error(FText::FromString(*CopyErrors));
-			ErrorMessage += FText::Format(LOCTEXT( "MigratePackages_CopyErrors", "Copied {0} files. Some content could not be copied."), FText::AsNumber(SuccessfullyCopiedPackages.Num())).ToString();
+			FString ErrorMessage;
+			Severity = EMessageSeverity::Error;
+			if( CopyErrors.Len() > 0 )
+			{
+				MigrateLog.NewPage( LOCTEXT("MigratePackages_CopyErrorsPage", "Copy Errors") );
+				MigrateLog.Error(FText::FromString(*CopyErrors));
+				ErrorMessage += FText::Format(LOCTEXT( "MigratePackages_CopyErrors", "Copied {0} files. Some content could not be copied."), FText::AsNumber(SuccessfullyCopiedPackages.Num())).ToString();
+			}
+			if( SourceControlErrors.Len() > 0 )
+			{
+				MigrateLog.NewPage( LOCTEXT("MigratePackages_SourceControlErrorsListPage", "Source Control Errors") );
+				MigrateLog.Error(FText::FromString(*SourceControlErrors));
+				ErrorMessage += LINE_TERMINATOR;
+				ErrorMessage += LOCTEXT( "MigratePackages_SourceControlErrorsList", "Some files reported source control errors.").ToString();
+			}
+			if ( SuccessfullyCopiedPackages.Num() > 0 )
+			{
+				MigrateLog.NewPage( LOCTEXT("MigratePackages_CopyErrorsSuccesslistPage", "Copied Successfully") );
+				MigrateLog.Info(FText::FromString(*SourceControlErrors));
+				ErrorMessage += LINE_TERMINATOR;
+				ErrorMessage += LOCTEXT( "MigratePackages_CopyErrorsSuccesslist", "Some files were copied successfully.").ToString();
+				for ( auto FileIt = SuccessfullyCopiedPackages.CreateConstIterator(); FileIt; ++FileIt )
+				{
+					if(FileIt->Len()>0)
+					{
+						MigrateLog.Info(FText::FromString(*FileIt));
+					}
+				}
+			}
+			LogMessage = FText::FromString(ErrorMessage);
 		}
-		if( SourceControlErrors.Len() > 0 )
+		else if (bUserCanceled)
 		{
-			MigrateLog.NewPage( LOCTEXT("MigratePackages_SourceControlErrorsListPage", "Source Control Errors") );
-			MigrateLog.Error(FText::FromString(*SourceControlErrors));
-			ErrorMessage += LINE_TERMINATOR;
-			ErrorMessage += LOCTEXT( "MigratePackages_SourceControlErrorsList", "Some files reported source control errors.").ToString();
+			LogMessage = LOCTEXT("MigratePackages_CanceledPage", "Content migration was canceled.");
 		}
-		if ( SuccessfullyCopiedPackages.Num() > 0 )
+		else if ( !bUserCanceled )
 		{
-			MigrateLog.NewPage( LOCTEXT("MigratePackages_CopyErrorsSuccesslistPage", "Copied Successfully") );
-			MigrateLog.Info(FText::FromString(*SourceControlErrors));
-			ErrorMessage += LINE_TERMINATOR;
-			ErrorMessage += LOCTEXT( "MigratePackages_CopyErrorsSuccesslist", "Some files were copied successfully.").ToString();
+			MigrateLog.NewPage( LOCTEXT("MigratePackages_CompletePage", "Content migration completed successfully!") );
 			for ( auto FileIt = SuccessfullyCopiedPackages.CreateConstIterator(); FileIt; ++FileIt )
 			{
 				if(FileIt->Len()>0)
@@ -3713,31 +4532,16 @@ void UAssetToolsImpl::MigratePackages_ReportConfirmed(TSharedPtr<TArray<ReportPa
 				}
 			}
 		}
-		LogMessage = FText::FromString(ErrorMessage);
-	}
-	else if (bUserCanceled)
-	{
-		LogMessage = LOCTEXT("MigratePackages_CanceledPage", "Content migration was canceled.");
-	}
-	else if ( !bUserCanceled )
-	{
-		MigrateLog.NewPage( LOCTEXT("MigratePackages_CompletePage", "Content migration completed successfully!") );
-		for ( auto FileIt = SuccessfullyCopiedPackages.CreateConstIterator(); FileIt; ++FileIt )
-		{
-			if(FileIt->Len()>0)
-			{
-				MigrateLog.Info(FText::FromString(*FileIt));
-			}
-		}
+		MigrateLog.Notify(LogMessage, Severity, true);
 	}
-	MigrateLog.Notify(LogMessage, Severity, true);
 }
 
 void UAssetToolsImpl::RecursiveGetDependencies(const FName& PackageName, TSet<FName>& AllDependencies, TSet<FString>& OutExternalObjectsPaths) const
 {
 	FAssetRegistryModule& AssetRegistryModule = FModuleManager::Get().LoadModuleChecked<FAssetRegistryModule>(TEXT("AssetRegistry"));
 	TArray<FName> Dependencies;
-	AssetRegistryModule.Get().GetDependencies(PackageName, Dependencies);
+	IAssetRegistry& AssetRegistry = AssetRegistryModule.Get();
+	AssetRegistry.GetDependencies(PackageName, Dependencies);
 	
 	for ( auto DependsIt = Dependencies.CreateConstIterator(); DependsIt; ++DependsIt )
 	{
@@ -3745,7 +4549,11 @@ void UAssetToolsImpl::RecursiveGetDependencies(const FName& PackageName, TSet<FN
 
 		const bool bIsEnginePackage = DependencyName.StartsWith(TEXT("/Engine"));
 		const bool bIsScriptPackage = DependencyName.StartsWith(TEXT("/Script"));
-		if ( !bIsEnginePackage && !bIsScriptPackage )
+
+		// The asset registry can give some reference to some deleted assets. We don't want to migrate these.
+		const bool bAssetExist = AssetRegistry.GetAssetPackageDataCopy(*DependsIt).IsSet();
+
+		if ( !bIsEnginePackage && !bIsScriptPackage && bAssetExist)
 		{
 			if (!AllDependencies.Contains(*DependsIt))
 			{
@@ -3755,10 +4563,13 @@ void UAssetToolsImpl::RecursiveGetDependencies(const FName& PackageName, TSet<FN
 		}
 	}
 
-	// Handle Specific External Actors use case (only used for the Migrate path for now)
-	// todo: revisit how to handle those in a more generic way.
+	// Handle Specific External Objects use case (only used for the Migrate path for now)
+	// todo: revisit how to handle those in a more generic way. Should the FWorldPartitionAssetDependencyGatherer handle the external objects reference also?
 	TArray<FAssetData> Assets;
-	if (AssetRegistryModule.Get().GetAssetsByPackageName(PackageName, Assets))
+
+	// The migration only work on the saved version of the assets so no need to scan the for the in memory only assets. This also greatly improve the performance of the migration when a lot assets are loaded in the editor.
+	const bool bOnlyIncludeOnDiskAssets = true;
+	if (AssetRegistryModule.Get().GetAssetsByPackageName(PackageName, Assets, bOnlyIncludeOnDiskAssets))
 	{
 		for (const FAssetData& AssetData : Assets)
 		{
@@ -3773,7 +4584,7 @@ void UAssetToolsImpl::RecursiveGetDependencies(const FName& PackageName, TSet<FN
 						AssetRegistryModule.Get().ScanPathsSynchronous({ ExternalObjectsPath }, /*bForceRescan*/true, /*bIgnoreDenyListScanFilters*/true);
 
 						TArray<FAssetData> ExternalObjectAssets;
-						AssetRegistryModule.Get().GetAssetsByPath(FName(*ExternalObjectsPath), ExternalObjectAssets, /*bRecursive*/true);
+						AssetRegistryModule.Get().GetAssetsByPath(FName(*ExternalObjectsPath), ExternalObjectAssets, /*bRecursive*/true, bOnlyIncludeOnDiskAssets);
 
 						for (const FAssetData& ExternalObjectAsset : ExternalObjectAssets)
 						{
diff --git a/Engine/Source/Developer/AssetTools/Private/AssetTools.h b/Engine/Source/Developer/AssetTools/Private/AssetTools.h
--- a/Engine/Source/Developer/AssetTools/Private/AssetTools.h
+++ b/Engine/Source/Developer/AssetTools/Private/AssetTools.h
@@ -107,6 +107,7 @@ public:
 	virtual bool CreateDiffProcess(const FString& DiffCommand, const FString& OldTextFilename, const FString& NewTextFilename, const FString& DiffArgs = FString("")) const override;
 	virtual void MigratePackages(const TArray<FName>& PackageNamesToMigrate) const override;
 	virtual void MigratePackages(const TArray<FName>& PackageNamesToMigrate, const FString& TargetPath, const struct FMigrationOptions& Options = FMigrationOptions()) const override;
+	virtual UE::AssetTools::FOnPackageMigration& GetOnPackageMigration() override;
 	virtual void BeginAdvancedCopyPackages(const TArray<FName>& InputNamesToCopy, const FString& TargetPath) const override;
 	virtual void FixupReferencers(const TArray<UObjectRedirector*>& Objects, bool bCheckoutDialogPrompt = true, ERedirectFixupMode FixupMode = ERedirectFixupMode::DeleteFixedUpRedirectors) const override;
 	virtual bool IsFixupReferencersInProgress() const override;
@@ -241,6 +242,8 @@ private:
 	bool CreateAssetsAsExternallyReferenceable;
 
 	TMap<FName, FIsNameAllowed> IsNameAllowedDelegates;
+
+	UE::AssetTools::FOnPackageMigration OnPackageMigration;
 };
 
 PRAGMA_ENABLE_DEPRECATION_WARNINGS
diff --git a/Engine/Source/Developer/AssetTools/Private/PackageMigrationContext.cpp b/Engine/Source/Developer/AssetTools/Private/PackageMigrationContext.cpp
new file mode 100644
--- /dev/null
+++ b/Engine/Source/Developer/AssetTools/Private/PackageMigrationContext.cpp
@@ -0,0 +1,190 @@
+// Copyright Epic Games, Inc. All Rights Reserved.
+
+#include "PackageMigrationContext.h"
+
+#include "Misc/AssertionMacros.h"
+#include "Misc/PackageName.h"
+#include "Misc/Paths.h"
+#include "UObject/Package.h"
+
+namespace UE::AssetTools
+{
+
+namespace Private::Migration
+{
+	static FString MigrationDir(TEXT("Migration"));
+}
+
+FPackageMigrationContext::EPackageMigrationStep FPackageMigrationContext::GetCurrentStep() const
+{
+	return CurrentStep;
+}
+
+void FPackageMigrationContext::AddSucessfullMigrationMessage(const FText& InMessage)
+{
+	MigratedPackageMessages.Add(InMessage);
+}
+
+void FPackageMigrationContext::AddWarningMigrationMessage(const FText& InMessage)
+{
+	WarningMessage.Add(InMessage);
+}
+
+void FPackageMigrationContext::AddErrorMigrationMessage(const FText& InMessage)
+{
+	ErrorMessages.Add(InMessage);
+}
+
+const FPackageMigrationContext::FScopedMountPoint& FPackageMigrationContext::GetDestinationMountPoint() const
+{
+	return DestinationMountPoint;
+}
+
+const TArray<FPackageMigrationContext::FMigrationPackageData>& FPackageMigrationContext::GetMigrationPackagesData() const
+{
+	return MigrationPackagesData;
+}
+
+const TArray<UPackage*>& FPackageMigrationContext::GetMovedOutOfTheWayPackages() const
+{
+	return PackagesThatWhereMoved;
+}
+
+void FPackageMigrationContext::MoveInTheWayPackage(UPackage* Package)
+{
+	check(Package);
+
+	checkf(CurrentStep < EPackageMigrationStep::InTheWayPackagesMoved, TEXT("The migration should not move existing pacakges after the broadcast of the in the way pacakge moved"));
+
+	TemporalyMovedPackages.Emplace(Package);
+	PackagesThatWhereMoved.Add(Package);
+}
+
+FPackageMigrationContext::FPackageMigrationContext(FScopedMountPoint&& InMountPoint)
+	: DestinationMountPoint(MoveTemp(InMountPoint))
+{
+}
+
+void FPackageMigrationContext::AddReferencedObjects(FReferenceCollector& Collector)
+{
+	Collector.AddReferencedObjects(PackagesThatWhereMoved);
+
+	for (FScopedTemporalyMovedPackage& MovedPackage : TemporalyMovedPackages)
+	{
+		Collector.AddReferencedObject(MovedPackage.PackageToMove);
+	}
+
+	for (FMigrationPackageData& PacakgeData : MigrationPackagesData)
+	{
+		Collector.AddReferencedObject(PacakgeData.InstancedPackage);
+	}
+}
+
+FString FPackageMigrationContext::GetReferencerName() const
+{
+	return TEXT("FPackageMigrationContext");
+}
+
+const FString& FPackageMigrationContext::FScopedMountPoint::GetRootPath() const
+{
+	return RootPath;
+}
+
+FString FPackageMigrationContext::FScopedMountPoint::GetMigratedPackageFilename(const FString& LongPackageName, const FStringView& InExtension) const
+{
+	FString PackageRoot;
+	FString PackagePath;
+	FString PackageName;
+
+	FPackageName::SplitLongPackageName(LongPackageName, PackageRoot, PackagePath, PackageName);
+
+	FString PackageFilename;
+	PackageFilename.Reserve(ContentPath.Len() + PackagePath.Len() + PackageName.Len() + InExtension.Len());
+	PackageFilename.Append(ContentPath);
+	PackageFilename.Append(PackagePath);
+	PackageFilename.Append(PackageName);
+	PackageFilename.Append(InExtension);
+
+	return PackageFilename;
+}
+
+FString FPackageMigrationContext::FScopedMountPoint::GetNewPackageNameForMigration(const FString& LongPackageName) const
+{
+	const bool InWithoutSlashes = false;
+	FName CurrentPackageMountPoint = FPackageName::GetPackageMountPoint(LongPackageName, InWithoutSlashes);
+
+	FStringView CurrentPackageWithoutMountPoint = LongPackageName;
+	CurrentPackageWithoutMountPoint.RightChopInline(CurrentPackageMountPoint.GetStringLength());
+
+
+	FString NewPackageName;
+	NewPackageName.Reserve(GetRootPath().Len() - CurrentPackageWithoutMountPoint.Len());
+	NewPackageName.Append(GetRootPath());
+	NewPackageName.Append(CurrentPackageWithoutMountPoint);
+
+	return NewPackageName;
+}
+
+FPackageMigrationContext::FScopedMountPoint::FScopedMountPoint(FString&& InRootPath, FString&& InContentPath)
+	: RootPath(MoveTemp(InRootPath))
+	, ContentPath(MoveTemp(InContentPath))
+{
+	if (!FPackageName::MountPointExists(RootPath))
+	{
+		bHasMountedANewMountPoint = true;
+		// Use a temp folder to avoid letting the asset registry scan the destination
+		FPackageName::RegisterMountPoint(RootPath, FPaths::Combine(FPaths::ProjectIntermediateDir(), UE::AssetTools::Private::Migration::MigrationDir));
+	}
+}
+
+FPackageMigrationContext::FScopedMountPoint::~FScopedMountPoint()
+{
+	// Empty root path means that we didn't add or own a new mount point
+	if (bHasMountedANewMountPoint && !RootPath.IsEmpty() && !ContentPath.IsEmpty())
+	{
+		FPackageName::UnRegisterMountPoint(RootPath, FPaths::Combine(FPaths::ProjectIntermediateDir(), UE::AssetTools::Private::Migration::MigrationDir));
+	}
+}
+
+FPackageMigrationContext::FScopedTemporalyMovedPackage::FScopedTemporalyMovedPackage(UPackage* InPackageToMove)
+	: PackageToMove(InPackageToMove)
+	, OriginalName(PackageToMove->GetName())
+{
+	const ERenameFlags PkgRenameFlags = REN_ForceNoResetLoaders | REN_DoNotDirty | REN_DontCreateRedirectors | REN_NonTransactional | REN_SkipGeneratedClasses;
+	check(PackageToMove->Rename(*MakeUniqueObjectName(nullptr, UPackage::StaticClass(), *FString::Printf(TEXT("%s_MOVEDFORMIGRATION"), *PackageToMove->GetName())).ToString(), nullptr, PkgRenameFlags));
+}
+
+FPackageMigrationContext::FScopedTemporalyMovedPackage::~FScopedTemporalyMovedPackage()
+{
+	const ERenameFlags PkgRenameFlags = REN_ForceNoResetLoaders | REN_DoNotDirty | REN_DontCreateRedirectors | REN_NonTransactional | REN_SkipGeneratedClasses;
+	check(PackageToMove->Rename(*OriginalName, nullptr, PkgRenameFlags));
+}
+
+const FString& FPackageMigrationContext::FMigrationPackageData::GetInstancedPackageName() const
+{
+	return InstancedPackageName;
+}
+
+const FString& FPackageMigrationContext::FMigrationPackageData::GetOriginalPackageName() const
+{
+	return OriginalPackageName;
+}
+
+const FString& FPackageMigrationContext::FMigrationPackageData::GetDestinationFilename() const
+{
+	return DestinationFilename;
+}
+
+UPackage* FPackageMigrationContext::FMigrationPackageData::GetInstancedPackage() const
+{
+	return InstancedPackage;
+}
+
+FPackageMigrationContext::FMigrationPackageData::FMigrationPackageData(const FString& InInstancedPackageName, const FString& InOriginalPackageName, const FString& InDestinationFilename)
+	: InstancedPackageName(InInstancedPackageName)
+	, OriginalPackageName(InOriginalPackageName)
+	, DestinationFilename(InDestinationFilename)
+{
+}
+
+}
diff --git a/Engine/Source/Developer/AssetTools/Public/IAssetTools.h b/Engine/Source/Developer/AssetTools/Public/IAssetTools.h
--- a/Engine/Source/Developer/AssetTools/Public/IAssetTools.h
+++ b/Engine/Source/Developer/AssetTools/Public/IAssetTools.h
@@ -18,6 +18,7 @@
 #include "IAssetTools.generated.h"
 
 struct FAssetData;
+class IAssetTools;
 class IAssetTypeActions;
 class IClassTypeActions;
 class UFactory;
@@ -26,6 +27,14 @@ class UAdvancedCopyCustomization;
 class FNamePermissionList;
 class FPathPermissionList;
 
+namespace UE::AssetTools
+{
+	// Declared in PackageMigrationContext.h
+	struct FPackageMigrationContext;
+
+	DECLARE_MULTICAST_DELEGATE_OneParam(FOnPackageMigration, FPackageMigrationContext&);
+}
+
 UENUM()
 enum class EAssetClassAction : uint8
 {
@@ -202,7 +211,7 @@ struct FMigrationOptions
 	UPROPERTY(BlueprintReadWrite, Category = MigrationOptions)
 	EAssetMigrationConflict AssetConflict;
 
-	/** Destination for assets that don't have a corresponding content folder. If left empty those assets are not migrated. */
+	/** Destination for assets that don't have a corresponding content folder. If left empty those assets are not migrated. (Not used by the new migration)*/
 	UPROPERTY(BlueprintReadWrite, Category = MigrationOptions)
 	FString OrphanFolder;
 
@@ -458,6 +467,12 @@ public:
 	UFUNCTION(BlueprintCallable, Category = "Editor Scripting | Asset Tools")
 	virtual void MigratePackages(const TArray<FName>& PackageNamesToMigrate, const FString& DestinationPath, const struct FMigrationOptions& Options = FMigrationOptions()) const = 0;
 
+	/**
+	 * Event called when some packages are migrated
+	 * Note this is only true when AssetTools.UseNewPackageMigration is true
+	 */
+	virtual UE::AssetTools::FOnPackageMigration& GetOnPackageMigration() = 0;
+
 	/* Copy packages and dependencies to another folder */
 	virtual void BeginAdvancedCopyPackages(const TArray<FName>& InputNamesToCopy, const FString& TargetPath) const = 0;
 
diff --git a/Engine/Source/Developer/AssetTools/Public/PackageMigrationContext.h b/Engine/Source/Developer/AssetTools/Public/PackageMigrationContext.h
new file mode 100644
--- /dev/null
+++ b/Engine/Source/Developer/AssetTools/Public/PackageMigrationContext.h
@@ -0,0 +1,185 @@
+// Copyright Epic Games, Inc. All Rights Reserved.
+
+#pragma once
+
+#include "Internationalization/Text.h"
+#include "Containers/UnrealString.h"
+#include "UObject/GCObject.h"
+
+class UPackage;
+class UAssetToolsImpl;
+
+namespace UE::AssetTools
+{
+
+namespace Private
+{
+	struct FPackageMigrationImpl;
+}
+
+struct ASSETTOOLS_API FPackageMigrationContext : public FGCObject
+{
+	enum class EPackageMigrationStep;
+
+	EPackageMigrationStep GetCurrentStep() const;
+
+	// Notify the user that files where migrated. It should tell the user what was migrated and where it is in the destination.
+	void AddSucessfullMigrationMessage(const FText& InMessage);
+
+	// Notify the user that some asset may require some extra work to fully migrate
+	void AddWarningMigrationMessage(const FText& InMessage);
+
+	// Notify the user that some asset couldn't be migrated
+	void AddErrorMigrationMessage(const FText& InMessage);
+
+	struct FScopedMountPoint;
+	const FScopedMountPoint& GetDestinationMountPoint() const;
+
+	struct FMigrationPackageData;
+	// Return the data associated to a package that is migrated
+	const TArray<FMigrationPackageData>& GetMigrationPackagesData() const;
+
+	// Return the array of non instanced package where moved because they were in the way of the migration
+	const TArray<UPackage*>& GetMovedOutOfTheWayPackages() const;
+
+	enum class EPackageMigrationStep
+	{
+		// Called before the migration begin
+		BeginMigration,
+		// Called after the non external packages to migrate where processed
+		PostAssetMigrationPackageDataCreated,
+		// Called after the non external packages to migrate where processed
+		PostExternalMigrationPackageDataCreated,
+		// Called after all the package that might be in the way for the migration where move
+		InTheWayPackagesMoved,
+		// Called after the instanced package where created with their load linker setup being already set to load the right file
+		InstancedPackagesCreated,
+		// Called after the instanced package where loaded
+		InstancedPackagesLoaded,
+		// Called after the instanced package were saved into the destination
+		InstancedPackagesSaved,
+		// Called after the instanced package were removed
+		PostCleaningInstancedPackages,
+		// Called at the end of the migration. After the in the way package where restored but before the log is processed
+		EndMigration
+	};
+
+	// Utility to manage the temp mount point if needed and to help convert the paths during the migration
+	struct ASSETTOOLS_API FScopedMountPoint
+	{
+	public:
+		FScopedMountPoint() = delete;
+		FScopedMountPoint(const FScopedMountPoint&) = delete;
+		FScopedMountPoint& operator=(const FScopedMountPoint&) = delete;
+
+		FScopedMountPoint(FScopedMountPoint&&) = default;
+		FScopedMountPoint& operator=(FScopedMountPoint&&) = default;
+
+		const FString& GetRootPath() const;
+		const FString& GetContentPath() const;
+
+		FString GetMigratedPackageFilename(const FString& LongPackageName, const FStringView& InExtension) const;
+
+		// Get name of package when moved under the destination mount point
+		FString GetNewPackageNameForMigration(const FString& LongPackageName) const;
+
+	private:
+		friend UAssetToolsImpl;
+		friend FPackageMigrationContext;
+		friend UE::AssetTools::Private::FPackageMigrationImpl;
+
+		FScopedMountPoint(FString&& InRootPath, FString&& InContentPath);
+		~FScopedMountPoint();
+
+		FString RootPath;
+		FString ContentPath;
+		bool bHasMountedANewMountPoint = false;
+	};
+
+	// The data associated to an package that take part of the migration process.
+	struct ASSETTOOLS_API FMigrationPackageData
+	{
+		const FString& GetInstancedPackageName() const;
+		const FString& GetOriginalPackageName() const;
+		const FString& GetDestinationFilename() const;
+
+		UPackage* GetInstancedPackage() const;
+
+		// Tell the migration that package must be loaded
+		bool bNeedInstancedLoad = false;
+
+		// Tell the migration that the package must migrated by save
+		bool bNeedToBeSaveMigrated = false;
+
+		FMigrationPackageData(const FString& InInstancedPackageName, const FString& InOriginalPackageName, const FString& InDestinationFilename);
+
+	private:
+		friend UAssetToolsImpl;
+		friend FPackageMigrationContext;
+		friend UE::AssetTools::Private::FPackageMigrationImpl;
+
+		FString InstancedPackageName;
+		FString OriginalPackageName;
+
+		FString DestinationFilename;
+		UPackage* InstancedPackage = nullptr;
+	};
+
+	FPackageMigrationContext(const FPackageMigrationContext&) = delete;
+	FPackageMigrationContext(FPackageMigrationContext&&) = delete;
+	FPackageMigrationContext operator=(const FPackageMigrationContext&) = delete;
+	FPackageMigrationContext operator=(FPackageMigrationContext&&) = delete;
+
+
+	virtual void AddReferencedObjects( FReferenceCollector& Collector ) override;
+
+	virtual FString GetReferencerName() const override;
+
+private:
+
+	friend UAssetToolsImpl;
+	friend UE::AssetTools::Private::FPackageMigrationImpl;
+
+	// Move a package out of the way for the duration of the migration
+	void MoveInTheWayPackage(UPackage* Package);
+
+	FPackageMigrationContext(FScopedMountPoint&& InMountPoint);
+
+	~FPackageMigrationContext() = default;
+
+	// Helper to move existing package out of the way during the migration
+	struct ASSETTOOLS_API FScopedTemporalyMovedPackage
+	{
+	public:
+		FScopedTemporalyMovedPackage(UPackage* InPackageToMove);
+
+		~FScopedTemporalyMovedPackage();
+
+		UPackage* PackageToMove = nullptr;
+		FString OriginalName;
+	};
+
+	EPackageMigrationStep CurrentStep;
+
+	// Message of the for the successfully migrated asset
+	TArray<FText> MigratedPackageMessages;
+
+	// Message to warn the user that some extra step may be required from him to complete the migration
+	TArray<FText> WarningMessage;
+
+	// General Error messages of the migration
+	TArray<FText> ErrorMessages;
+
+	FScopedMountPoint DestinationMountPoint;
+
+	// Helper for the package that were moved because they were in the way of migrated packages
+	TArray<FScopedTemporalyMovedPackage> TemporalyMovedPackages;
+
+	// The package that have been moved during the migration
+	TArray<UPackage*> PackagesThatWhereMoved;
+
+	// The data associated to the packages that take part of the migration process.
+	TArray<FMigrationPackageData> MigrationPackagesData;
+};
+
+}
\ No newline at end of file
diff --git a/Engine/Source/Editor/StringTableEditor/Private/StringTableEditorModule.cpp b/Engine/Source/Editor/StringTableEditor/Private/StringTableEditorModule.cpp
--- a/Engine/Source/Editor/StringTableEditor/Private/StringTableEditorModule.cpp
+++ b/Engine/Source/Editor/StringTableEditor/Private/StringTableEditorModule.cpp
@@ -3,8 +3,10 @@
 #include "StringTableEditorModule.h"
 #include "Modules/ModuleManager.h"
 #include "IStringTableEditor.h"
+#include "Internationalization/StringTableRegistry.h"
 #include "StringTableEditor.h"
 #include "AssetTypeActions_StringTable.h"
+#include "PackageMigrationContext.h"
 
 IMPLEMENT_MODULE(FStringTableEditorModule, StringTableEditor);
 
@@ -17,10 +19,17 @@ void FStringTableEditorModule::StartupModule()
 
 	IAssetTools& AssetTools = FModuleManager::LoadModuleChecked<FAssetToolsModule>("AssetTools").Get();
 	AssetTools.RegisterAssetTypeActions(MakeShareable(new FAssetTypeActions_StringTable()));
+
+	AssetTools.GetOnPackageMigration().AddRaw(this, &FStringTableEditorModule::OnPackageMigration);
 }
 
 void FStringTableEditorModule::ShutdownModule()
 {
+	if (FAssetToolsModule* AssetToolsModule = FModuleManager::GetModulePtr<FAssetToolsModule>("AssetTools"))
+	{
+		AssetToolsModule->Get().GetOnPackageMigration().RemoveAll(this);
+	}
+
 	MenuExtensibilityManager.Reset();
 	ToolBarExtensibilityManager.Reset();
 }
@@ -31,3 +40,35 @@ TSharedRef<IStringTableEditor> FStringTableEditorModule::CreateStringTableEditor
 	NewStringTableEditor->InitStringTableEditor(Mode, InitToolkitHost, StringTable);
 	return NewStringTableEditor;
 }
+
+void FStringTableEditorModule::OnPackageMigration(UE::AssetTools::FPackageMigrationContext& MigrationContext)
+{
+	constexpr auto ForEachStringTableInTheWay = [](UE::AssetTools::FPackageMigrationContext& MigrationContext, TFunctionRef<void(UStringTable*)> Callback)
+	{
+		const TArray<UPackage*>& MovedPackages = MigrationContext.GetMovedOutOfTheWayPackages();
+		for (UPackage* Package : MovedPackages)
+		{
+			UObject* MainAsset = Package->FindAssetInPackage();
+			if (UStringTable* StringTableAsset = Cast<UStringTable>(MainAsset))
+			{
+				Callback(StringTableAsset);
+			}
+		}
+	};
+
+	// The migration can create temporally new packages that will use the same name as the asset migrated. So we need to unregister the migrated asset for the duration of the migration
+	if (MigrationContext.GetCurrentStep() == UE::AssetTools::FPackageMigrationContext::EPackageMigrationStep::InTheWayPackagesMoved)
+	{
+		ForEachStringTableInTheWay(MigrationContext, [](UStringTable* StringTableAsset)
+			{
+				FStringTableRegistry::Get().UnregisterStringTable(StringTableAsset->GetStringTableId());
+			});
+	}
+	else if (MigrationContext.GetCurrentStep() == UE::AssetTools::FPackageMigrationContext::EPackageMigrationStep::EndMigration)
+	{
+		ForEachStringTableInTheWay(MigrationContext, [](UStringTable* StringTableAsset)
+			{
+				FStringTableRegistry::Get().RegisterStringTable(StringTableAsset->GetStringTableId(), StringTableAsset->GetMutableStringTable());
+			});
+	}
+}
diff --git a/Engine/Source/Editor/StringTableEditor/Public/StringTableEditorModule.h b/Engine/Source/Editor/StringTableEditor/Public/StringTableEditorModule.h
--- a/Engine/Source/Editor/StringTableEditor/Public/StringTableEditorModule.h
+++ b/Engine/Source/Editor/StringTableEditor/Public/StringTableEditorModule.h
@@ -10,6 +10,11 @@
 class IStringTableEditor;
 class UStringTable;
 
+namespace UE::AssetTools
+{
+	struct FPackageMigrationContext;
+}
+
 /** String Table Editor module */
 class STRINGTABLEEDITOR_API FStringTableEditorModule : public IModuleInterface,
 	public IHasMenuExtensibility, public IHasToolBarExtensibility
@@ -39,6 +44,8 @@ public:
 	static const FName StringTableEditorAppIdentifier;
 
 private:
+	void OnPackageMigration(UE::AssetTools::FPackageMigrationContext& MigrationContext);
+
 	TSharedPtr<FExtensibilityManager> MenuExtensibilityManager;
 	TSharedPtr<FExtensibilityManager> ToolBarExtensibilityManager;
 };
diff --git a/Engine/Source/Editor/StringTableEditor/StringTableEditor.Build.cs b/Engine/Source/Editor/StringTableEditor/StringTableEditor.Build.cs
--- a/Engine/Source/Editor/StringTableEditor/StringTableEditor.Build.cs
+++ b/Engine/Source/Editor/StringTableEditor/StringTableEditor.Build.cs
@@ -18,12 +18,8 @@ public class StringTableEditor : ModuleRules
                 "DesktopPlatform",
 				"EditorFramework",
                 "UnrealEd",
-            });
-
-        PrivateIncludePathModuleNames.AddRange(
-            new string[] {
-                "AssetTools"
-            });
+				"AssetTools",
+			});
 
         DynamicallyLoadedModuleNames.Add("WorkspaceMenuStructure");
     }
diff --git a/Engine/Source/Runtime/Engine/Classes/Engine/Level.h b/Engine/Source/Runtime/Engine/Classes/Engine/Level.h
--- a/Engine/Source/Runtime/Engine/Classes/Engine/Level.h
+++ b/Engine/Source/Runtime/Engine/Classes/Engine/Level.h
@@ -1172,6 +1172,11 @@ public:
 	 */
 	static ENGINE_API FString GetExternalActorsPath(UPackage* InLevelPackage, const FString& InPackageShortName = FString());
 
+	/**
+	 * Extract the packaging Scheme used by an external actor package based on the name of the package
+	 */
+	static ENGINE_API EActorPackagingScheme GetActorPackagingSchemeFromActorPackageName(const FStringView InActorPackageName);
+
 	/**
 	 * Scans/Updates all Level Assets (level package and external packages)
 	 */
diff --git a/Engine/Source/Runtime/Engine/Private/ExternalPackageHelper.cpp b/Engine/Source/Runtime/Engine/Private/ExternalPackageHelper.cpp
--- a/Engine/Source/Runtime/Engine/Private/ExternalPackageHelper.cpp
+++ b/Engine/Source/Runtime/Engine/Private/ExternalPackageHelper.cpp
@@ -8,7 +8,7 @@ FExternalPackageHelper::FOnObjectPackagingModeChanged FExternalPackageHelper::On
 
 UPackage* FExternalPackageHelper::CreateExternalPackage(UObject* InObjectOuter, const FString& InObjectPath, EPackageFlags InFlags)
 {
-	UPackage* Package = CreatePackage(*FExternalPackageHelper::GetExternalPackageName(InObjectOuter->GetPackage(), InObjectPath));
+	UPackage* Package = CreatePackage(*FExternalPackageHelper::GetExternalPackageName(InObjectOuter->GetPackage()->GetName(), InObjectPath));
 	Package->SetPackageFlags(InFlags);
 	return Package;
 }
@@ -89,7 +89,7 @@ FString FExternalPackageHelper::GetExternalObjectsPath(UPackage* InPackage, cons
 	return FExternalPackageHelper::GetExternalObjectsPath(InPackage->GetName(), InPackageShortName);
 }
 
-FString FExternalPackageHelper::GetExternalPackageName(UPackage* InOuterPackage, const FString& InObjectPath)
+FString FExternalPackageHelper::GetExternalPackageName(const FString& InOuterPackageName, const FString& InObjectPath)
 {
 	// Convert the object path to lowercase to make sure we get the same hash for case insensitive file systems
 	FString ObjectPath = InObjectPath.ToLower();
@@ -108,7 +108,7 @@ FString FExternalPackageHelper::GetExternalPackageName(UPackage* InOuterPackage,
 	FString GuidBase36 = PackageGuid.ToString(EGuidFormats::Base36Encoded);
 	check(GuidBase36.Len());
 
-	FString BaseDir = FExternalPackageHelper::GetExternalObjectsPath(InOuterPackage);
+	FString BaseDir = FExternalPackageHelper::GetExternalObjectsPath(InOuterPackageName);
 
 	TStringBuilderWithBuffer<TCHAR, NAME_SIZE> ObjectPackageName;
 	ObjectPackageName.Append(BaseDir);
diff --git a/Engine/Source/Runtime/Engine/Private/Level.cpp b/Engine/Source/Runtime/Engine/Private/Level.cpp
--- a/Engine/Source/Runtime/Engine/Private/Level.cpp
+++ b/Engine/Source/Runtime/Engine/Private/Level.cpp
@@ -81,6 +81,7 @@ Level.cpp: Level-related functions
 #include "UObject/LinkerLoad.h"
 #include "WorldPartition/WorldPartitionHelpers.h"
 #include "HAL/PlatformFileManager.h"
+#include "Misc/PathViews.h"
 #endif
 #include "WorldPartition/WorldPartition.h"
 #include "WorldPartition/WorldPartitionSubsystem.h"
@@ -3515,6 +3516,21 @@ FString ULevel::GetExternalActorsPath(UPackage* InLevelPackage, const FString& I
 	return GetExternalActorsPath(InLevelPackage->GetName(), InPackageShortName);
 }
 
+EActorPackagingScheme ULevel::GetActorPackagingSchemeFromActorPackageName(const FStringView InActorPackageName)
+{
+	// Use the fact that the end of an actor package path is a GUID that is encoded in a base36 and those are always 25 character long to determine the PackagingScheme.
+	FStringView ActorBaseFilename = FPathViews::GetBaseFilename(InActorPackageName);
+
+	if (ActorBaseFilename.Len() == 22)
+	{
+		return EActorPackagingScheme::Reduced;
+	}
+
+	check(ActorBaseFilename.Len() == 21);
+
+	return EActorPackagingScheme::Original;
+}
+
 void ULevel::ScanLevelAssets(const FString& InLevelPackageName)
 {
 	LevelAssetRegistryHelper::ScanLevelAssets(InLevelPackageName);
diff --git a/Engine/Source/Runtime/Engine/Private/LevelStreaming.cpp b/Engine/Source/Runtime/Engine/Private/LevelStreaming.cpp
--- a/Engine/Source/Runtime/Engine/Private/LevelStreaming.cpp
+++ b/Engine/Source/Runtime/Engine/Private/LevelStreaming.cpp
@@ -2132,6 +2132,20 @@ bool ULevelStreaming::IsValidStreamingLevel() const
 	if (!PIESession && !WorldAsset.IsNull())
 	{
 		FName WorldPackageName = GetWorldAssetPackageFName();
+
+		if (UPackage* WorldPackage = FindObjectFast<UPackage>(nullptr, WorldPackageName))
+		{
+			if (FLinkerLoad* Linker = WorldPackage->GetLinker())
+			{
+				/**
+				 * This support packages that were or will be instanced on load.
+				 * This might be redundant but it avoid changing the behavior of this function where a loaded package can still fail 
+				 * if it doesn't have on disk file associated to it.
+				 */
+				return FPackageName::DoesPackageExist(Linker->GetPackagePath());
+			}
+		}
+
 		FPackagePath WorldPackagePath;
 		return FPackagePath::TryFromPackageName(WorldPackageName, /* Out*/ WorldPackagePath) && FPackageName::DoesPackageExist(WorldPackagePath);
 	}
diff --git a/Engine/Source/Runtime/Engine/Public/ExternalPackageHelper.h b/Engine/Source/Runtime/Engine/Public/ExternalPackageHelper.h
--- a/Engine/Source/Runtime/Engine/Public/ExternalPackageHelper.h
+++ b/Engine/Source/Runtime/Engine/Public/ExternalPackageHelper.h
@@ -54,21 +54,23 @@ public:
 	 * @return the path
 	 */
 	static FString GetExternalObjectsPath(UPackage* InPackage, const FString& InPackageShortName = FString(), bool bTryUsingPackageLoadedPath = false);
-	
-	/**
-	 * Loads objects from an external package
-	 */
-	template<typename T>
-	static void LoadObjectsFromExternalPackages(UObject* InOuter, TFunctionRef<void(T*)> Operation);
 
-private:
+
 	/**
 	 * Get the external package name for this object
+	 * @param InOuterPackageName The name of the package of that contains the outer of the object.
 	 * @param InObjectPath the fully qualified object path, in the format: 'Outermost.Outer.Name'
 	 * @return the package name
 	 */
-	static FString GetExternalPackageName(UPackage* InOuterPackage, const FString& InObjectPath);
+	static FString GetExternalPackageName(const FString& InOuterPackageName, const FString& InObjectPath);
 
+	/**
+	 * Loads objects from an external package
+	 */
+	template<typename T>
+	static void LoadObjectsFromExternalPackages(UObject* InOuter, TFunctionRef<void(T*)> Operation);
+
+private:
 	/** Get the external object package instance name. */
 	static FString GetExternalObjectPackageInstanceName(const FString& OuterPackageName, const FString& ObjectPackageName);
 };
@@ -139,15 +141,34 @@ void FExternalPackageHelper::LoadObjectsFromExternalPackages(UObject* InOuter, T
 	const bool bInstanced = !PackageResourceName.IsNone() && (PackageResourceName != OuterPackage->GetFName());
 	if (bInstanced)
 	{
+		const FLinkerInstancingContext* OuterInstancingContext = nullptr;
+		if (FLinkerLoad* OuterLinker = InOuter->GetLinker())
+		{
+			OuterInstancingContext = &OuterLinker->GetInstancingContext();
+		}
+
 		InstancingContext.AddPackageMapping(PackageResourceName, OuterPackage->GetFName());
 
 		for (const FString& ObjectPackageName : ObjectPackageNames)
 		{
-			const FString InstancedName = GetExternalObjectPackageInstanceName(OuterPackage->GetName(), ObjectPackageName);
-			InstancingContext.AddPackageMapping(FName(*ObjectPackageName), FName(*InstancedName));
+			FName InstancedName;
+			
+			FName ObjectPackageFName = *ObjectPackageName;
+			if (OuterInstancingContext)
+			{
+				InstancedName = OuterInstancingContext->RemapPackage(ObjectPackageFName);
+			}
+
+			// Remap to the a instanced package if it wasn't remapped already by the outer instancing context
+			if (InstancedName == ObjectPackageFName || InstancedName.IsNone())
+			{
+				InstancedName = *GetExternalObjectPackageInstanceName(OuterPackage->GetName(), ObjectPackageName);
+			}
+
+			InstancingContext.AddPackageMapping(ObjectPackageFName, InstancedName);
 
 			// Create instance package
-			UPackage* InstancePackage = CreatePackage(*InstancedName);
+			UPackage* InstancePackage = CreatePackage(*InstancedName.ToString());
 			// Propagate RF_Transient
 			if (OuterPackage->HasAnyFlags(RF_Transient))
 			{
