[AssetTools]
* Fixed compile error with missing include

diff --git a/Engine/Source/Developer/AssetTools/Private/AssetTools.cpp b/Engine/Source/Developer/AssetTools/Private/AssetTools.cpp
--- a/Engine/Source/Developer/AssetTools/Private/AssetTools.cpp
+++ b/Engine/Source/Developer/AssetTools/Private/AssetTools.cpp
@@ -3,6 +3,7 @@
 #include "AssetTools.h"
 #include "Factories/Factory.h"
 #include "Factories/BlueprintFactory.h"
+#include "MaterialShared.h"
 #include "Misc/MessageDialog.h"
 #include "HAL/FileManager.h"
 #include "Misc/FileHelper.h"
