Added a missing include.

diff --git a/Engine/Source/Developer/AssetTools/Private/AssetTools.cpp b/Engine/Source/Developer/AssetTools/Private/AssetTools.cpp
--- a/Engine/Source/Developer/AssetTools/Private/AssetTools.cpp
+++ b/Engine/Source/Developer/AssetTools/Private/AssetTools.cpp
@@ -171,6 +171,7 @@
 #include "ExternalPackageHelper.h"
 #include "UObject/LinkerInstancingContext.h"
 #include "PackageMigrationContext.h"
+#include "ComponentReregisterContext.h"
 
 #if WITH_EDITOR
 #include "Subsystems/AssetEditorSubsystem.h"
