Fixed some duplicated localization keys

diff --git a/Engine/Source/Developer/AssetTools/Private/AssetTools.cpp b/Engine/Source/Developer/AssetTools/Private/AssetTools.cpp
--- a/Engine/Source/Developer/AssetTools/Private/AssetTools.cpp
+++ b/Engine/Source/Developer/AssetTools/Private/AssetTools.cpp
@@ -829,22 +829,22 @@ namespace UE::AssetTools::Private
 			{
 				if (PackageMigrationContext.WarningMessage.IsEmpty())
 				{
-					NotificationMessage = LOCTEXT("MigratePackages_SuccessNotifcation", "Content migration completed.");
+					NotificationMessage = LOCTEXT("MigratePackages_SuccessNotification", "Content migration completed.");
 				}
 				else
 				{
-					NotificationMessage = LOCTEXT("MigratePackages_SuccessNotifcation", "Content migration completed with some warnings.");
+					NotificationMessage = LOCTEXT("MigratePackages_SuccessNotificationWithWarning", "Content migration completed with some warnings.");
 				}
 			}
 			else
 			{
 				if (PackageMigrationContext.MigratedPackageMessages.IsEmpty())
 				{
-					NotificationMessage = LOCTEXT("MigratePackages_FailureNotifcation", "Content migration failed.");
+					NotificationMessage = LOCTEXT("MigratePackages_FailureNotification", "Content migration failed.");
 				}
 				else
 				{
-					NotificationMessage = LOCTEXT("MigratePackages_PartialFailureNotifcation", "Content migration failed but some asset have been migrated sucessfuly.");
+					NotificationMessage = LOCTEXT("MigratePackages_PartialFailureNotification", "Content migration failed but some asset have been migrated sucessfuly.");
 				}
 
 				MigrateLog.NewPage(LOCTEXT("MigratePackages_ErrorsPageLabel", "Migration Errors"));
@@ -862,7 +862,7 @@ namespace UE::AssetTools::Private
 				}
 				else
 				{
-					NotificationMessage = LOCTEXT("MigratePackages_Canceled", "Content migration was canceled but some package where already migrated.");
+					NotificationMessage = LOCTEXT("MigratePackages_CanceledWithMigratedAssets", "Content migration was canceled but some package where already migrated.");
 				}
 			}
 
