Fix for errors in https://horde.devtools.epicgames.com/job/635beeaae24ac41a63dee15a being caused in UE5-Main

diff --git a/Engine/Source/Developer/AssetTools/Public/AssetTypeActivationOpenedMethod.h b/Engine/Source/Developer/AssetTools/Public/AssetTypeActivationOpenedMethod.h
--- a/Engine/Source/Developer/AssetTools/Public/AssetTypeActivationOpenedMethod.h
+++ b/Engine/Source/Developer/AssetTools/Public/AssetTypeActivationOpenedMethod.h
@@ -12,7 +12,7 @@
 
 // Types of permissions allowed when attempting to open an asset in editor via activation (EAssetTypeActivationMethod)
 UENUM()
-enum class EAssetTypeActivationOpenedMethod : int32
+enum class EAssetTypeActivationOpenedMethod : uint8
 {
 	Edit,
 	View
diff --git a/Engine/Source/Developer/AssetTools/Public/IAssetTypeActions.h b/Engine/Source/Developer/AssetTools/Public/IAssetTypeActions.h
--- a/Engine/Source/Developer/AssetTools/Public/IAssetTypeActions.h
+++ b/Engine/Source/Developer/AssetTools/Public/IAssetTypeActions.h
@@ -10,7 +10,7 @@
 #include "IAssetTypeActions.generated.h"
 
 class IToolkitHost;
-enum class EAssetTypeActivationOpenedMethod : int32;
+enum class EAssetTypeActivationOpenedMethod : uint8;
 
 namespace EAssetTypeActivationMethod
 {
