Set the platform user to the default platform user when you preview force feedback effects. This ensures that Device Properties can be applied to the IInputInterface when you preview an effect from the editor.

#rnx

diff --git a/Engine/Source/Developer/AssetTools/Private/AssetTypeActions/AssetTypeActions_ForceFeedbackEffect.cpp b/Engine/Source/Developer/AssetTools/Private/AssetTypeActions/AssetTypeActions_ForceFeedbackEffect.cpp
--- a/Engine/Source/Developer/AssetTools/Private/AssetTypeActions/AssetTypeActions_ForceFeedbackEffect.cpp
+++ b/Engine/Source/Developer/AssetTools/Private/AssetTypeActions/AssetTypeActions_ForceFeedbackEffect.cpp
@@ -9,6 +9,7 @@
 #include "Widgets/Layout/SBox.h"
 #include "Widgets/Input/SButton.h"
 #include "Widgets/Images/SImage.h"
+#include "GenericPlatform/GenericPlatformInputDeviceMapper.h"
 
 #define LOCTEXT_NAMESPACE "AssetTypeActions"
 
@@ -140,6 +141,7 @@ void FAssetTypeActions_ForceFeedbackEffect::PlayEffect(UForceFeedbackEffect* Eff
 	{
 		PreviewForceFeedbackEffect.ForceFeedbackEffect = Effect;
 		PreviewForceFeedbackEffect.PlayTime = 0.f;
+		PreviewForceFeedbackEffect.PlatformUser = IPlatformInputDeviceMapper::Get().GetPrimaryPlatformUser();
 	}
 	else
 	{
