Update comment

diff --git a/Engine/Source/Developer/AssetTools/Private/AssetTools.cpp b/Engine/Source/Developer/AssetTools/Private/AssetTools.cpp
--- a/Engine/Source/Developer/AssetTools/Private/AssetTools.cpp
+++ b/Engine/Source/Developer/AssetTools/Private/AssetTools.cpp
@@ -4566,7 +4566,7 @@ void UAssetToolsImpl::RecursiveGetDependencies(const FName& PackageName, TSet<FN
 	}
 
 	// Handle Specific External Objects use case (only used for the Migrate path for now)
-	// todo: revisit how to handle those in a more generic way. Should the FWorldPartitionAssetDependencyGatherer handle the external objects reference also?
+	// todo: revisit how to handle those in a more generic way. Should the FExternalActorAssetDependencyGatherer handle the external objects reference also?
 	TArray<FAssetData> Assets;
 
 	// The migration only work on the saved version of the assets so no need to scan the for the in memory only assets. This also greatly improve the performance of the migration when a lot assets are loaded in the editor.
