AssetDefinition - This is the seed of a long term replacement of IAssetTypeActions.  IAssetTypeActions is woefully inadiquit now.  It puts a lot of functionality in one place, which wouldn't be a terrible thing, except for the GetActions, and some of the menu stuff being built into this class.  We need to make an effort to not have to load every single asset just to open the context menu you have selected, to-wit, a previous CL made it so that summoning the menu no longer *Had* to load the assets, however several remaining issues remained with the Blueprint Utility Actions, and the fact that we don't "KNOW" if we need to call GetActions for an asset.  We could introduce a replacement in the class, but that would just get inherited, causing other issues.  So instead the parallel UObject based AssetDefinition now exists.  In large part it will probably be very similar to the other one, but several things in it probably need to chnage.  For now there is a bridging concept in the form of the AssetDefinitionAssetTypeProxy.

Several amounts of effort went into the newer design of UAssetDefinition to make it easier to change over time, more things are structs now for input and output so additional parameters can be added.

AssetDefinition module is an editor module instead of being a developer module.
AssetDefinition does not include a myraid of things it shouldn't such as the blueprint module and UnrealEd.

diff --git a/Engine/Config/BaseEngine.ini b/Engine/Config/BaseEngine.ini
--- a/Engine/Config/BaseEngine.ini
+++ b/Engine/Config/BaseEngine.ini
@@ -1097,8 +1097,8 @@ DefaultWorldPartitionHLODLayer=/Engine/Maps/Templates/HLODs/HLODLayer_Instanced.
 +PropertyRedirects=(OldName="AnimationSettings.BoneNamesWithCustomAttributes",NewName="BoneNamesWithAnimationAttributes")
 
 ; 5.2
++StructRedirects=(OldName="/Script/AssetTools.RevisionInfo", NewName="/Script/AssetDefinition.RevisionInfo")
 +PropertyRedirects=(OldName="DecalComponent.SetFadeIn.Duaration", NewName="DecalComponent.SetFadeIn.Duration")
-
 +ClassRedirects=(OldName="/Script/MovieRenderPipelineCore.MoviePipelineMasterConfig", NewName="/Script/MovieRenderPipelineCore.MoviePipelinePrimaryConfig")
 
 [CoreUObject.Metadata]
diff --git a/Engine/Plugins/Animation/IKRig/Source/IKRigEditor/Private/RetargetEditor/AssetTypeActions_IKRetargeter.cpp b/Engine/Plugins/Animation/IKRig/Source/IKRigEditor/Private/RetargetEditor/AssetTypeActions_IKRetargeter.cpp
--- a/Engine/Plugins/Animation/IKRig/Source/IKRigEditor/Private/RetargetEditor/AssetTypeActions_IKRetargeter.cpp
+++ b/Engine/Plugins/Animation/IKRig/Source/IKRigEditor/Private/RetargetEditor/AssetTypeActions_IKRetargeter.cpp
@@ -79,49 +79,18 @@ void FAssetTypeActions_IKRetargeter::ExtendAnimSequenceToolMenu()
 
 void FAssetTypeActions_IKRetargeter::CreateRetargetSubMenu(FToolMenuSection& InSection)
 {
-	UContentBrowserAssetContextMenuContext* Context = InSection.FindContext<UContentBrowserAssetContextMenuContext>();
-	if (!Context)
-	{
-		return;			
-	}
-	
-	TArray<UObject*> SelectedObjects = Context->GetSelectedObjects();
-	if (SelectedObjects.IsEmpty())
-	{
-		return;
-	}
-
-	// change menu label if anim blueprint is selected
-	static const FText AnimAssetLabel(LOCTEXT("RetargetAnimationAsset", "Duplicate and Retarget Animation Assets"));
-	static const FText AnimBlueprintLabel(LOCTEXT("RetargetAnimationBlueprint", "Duplicate and Retarget Animation Blueprint"));
-
-	const FText MenuLabel = Cast<UAnimBlueprint>(SelectedObjects[0]) == nullptr ? AnimAssetLabel : AnimBlueprintLabel;
-	
 	InSection.AddMenuEntry(
 		"IKRetargetToDifferentSkeleton",
-		MenuLabel,
-		LOCTEXT("RetargetAnimation_ToolTip", "Duplicate an animation asset and retarget to a different skeleton."),
+		LOCTEXT("RetargetAnimation_Label", "Duplicate and Retarget Animation Assets/Blueprints"),
+		LOCTEXT("RetargetAnimation_ToolTip", "Duplicate an animation asset/blueprint and retarget to a different skeleton."),
 		FSlateIcon(FAppStyle::GetAppStyleSetName(), "GenericCurveEditor.TabIcon"),
-		FUIAction(FExecuteAction::CreateLambda([SelectedObjects]()
-		{
-			SRetargetAnimAssetsWindow::ShowWindow(SelectedObjects);
-		}),
-		FCanExecuteAction::CreateLambda([InSection]()
+		FToolMenuExecuteAction::CreateLambda([](const FToolMenuContext& InContext)
 		{
-			if (UContentBrowserAssetContextMenuContext* Context = InSection.FindContext<UContentBrowserAssetContextMenuContext>())
-			{
-				TArray<UObject*> SelectedObjects = Context->GetSelectedObjects();
-				for (UObject* SelectedObject : SelectedObjects)
-				{
-					if (Cast<UAnimationAsset>(SelectedObject) || Cast<UAnimBlueprint>(SelectedObject))
-					{
-						return true;
-					}
-				}		
-			}
-			
-			return false;
-		}))
+        	if (const UContentBrowserAssetContextMenuContext* Context = InContext.FindContext<UContentBrowserAssetContextMenuContext>())
+        	{
+        		SRetargetAnimAssetsWindow::ShowWindow(Context->LoadSelectedObjects<UObject>());
+        	}
+		})
 	);
 }
 
diff --git a/Engine/Plugins/Editor/AssetManagerEditor/Source/AssetManagerEditor/AssetManagerEditor.Build.cs b/Engine/Plugins/Editor/AssetManagerEditor/Source/AssetManagerEditor/AssetManagerEditor.Build.cs
--- a/Engine/Plugins/Editor/AssetManagerEditor/Source/AssetManagerEditor/AssetManagerEditor.Build.cs
+++ b/Engine/Plugins/Editor/AssetManagerEditor/Source/AssetManagerEditor/AssetManagerEditor.Build.cs
@@ -35,6 +35,7 @@ public class AssetManagerEditor : ModuleRules
 				"ContentBrowser",
 				"ContentBrowserData",
 				"WorkspaceMenuStructure",
+				"AssetDefinition",
 				"AssetTools",
 				"PropertyEditor",
 				"GraphEditor",
diff --git a/Engine/Plugins/Editor/AssetManagerEditor/Source/AssetManagerEditor/Private/AssetSourceControlContextMenu.cpp b/Engine/Plugins/Editor/AssetManagerEditor/Source/AssetManagerEditor/Private/AssetSourceControlContextMenu.cpp
--- a/Engine/Plugins/Editor/AssetManagerEditor/Source/AssetManagerEditor/Private/AssetSourceControlContextMenu.cpp
+++ b/Engine/Plugins/Editor/AssetManagerEditor/Source/AssetManagerEditor/Private/AssetSourceControlContextMenu.cpp
@@ -2,6 +2,8 @@
 
 #include "AssetSourceControlContextMenu.h"
 
+#include "AssetDefinition.h"
+#include "AssetDefinitionRegistry.h"
 #include "Algo/Transform.h"
 #include "AssetRegistry/AssetRegistryModule.h"
 #include "AssetToolsModule.h"
@@ -851,24 +853,27 @@ bool FAssetSourceControlContextMenuState::CanExecuteSCCMerge() const
 {
 	FAssetToolsModule& AssetToolsModule = FModuleManager::GetModuleChecked<FAssetToolsModule>("AssetTools");
 
-	bool bCanExecuteMerge = bCanExecuteSCCMerge;
-	for (const FAssetData& AssetData : SelectedAssets)
+	if (bCanExecuteSCCMerge && SelectedAssets.Num() > 0)
 	{
-		if (UObject* CurrentObject = AssetData.GetAsset())
+		for (const FAssetData& AssetData : SelectedAssets)
 		{
-			auto AssetTypeActions = AssetToolsModule.Get().GetAssetTypeActionsForClass(CurrentObject->GetClass()).Pin();
-			if (AssetTypeActions.IsValid())
+			if (const UAssetDefinition* AssetDefinition = UAssetDefinitionRegistry::Get()->GetAssetDefinitionForAsset(AssetData))
 			{
-				bCanExecuteMerge = AssetTypeActions->CanMerge();
+				if (!AssetDefinition->CanMerge())
+				{
+					return false;
+				}
+			}
+			else
+			{
+				return false;
 			}
 		}
-		else
-		{
-			bCanExecuteMerge = false;
-		}
+
+		return true;
 	}
 
-	return bCanExecuteMerge;
+	return false;
 }
 
 void FAssetSourceControlContextMenuState::ExecuteSCCRefresh() const
diff --git a/Engine/Plugins/Editor/ContentBrowser/ContentBrowserAssetDataSource/Source/ContentBrowserAssetDataSource/ContentBrowserAssetDataSource.Build.cs b/Engine/Plugins/Editor/ContentBrowser/ContentBrowserAssetDataSource/Source/ContentBrowserAssetDataSource/ContentBrowserAssetDataSource.Build.cs
--- a/Engine/Plugins/Editor/ContentBrowser/ContentBrowserAssetDataSource/Source/ContentBrowserAssetDataSource/ContentBrowserAssetDataSource.Build.cs
+++ b/Engine/Plugins/Editor/ContentBrowser/ContentBrowserAssetDataSource/Source/ContentBrowserAssetDataSource/ContentBrowserAssetDataSource.Build.cs
@@ -26,6 +26,7 @@ namespace UnrealBuildTool.Rules
 			PrivateDependencyModuleNames.AddRange(
 				new string[]
 				{
+					"AssetDefinition",
 					"ApplicationCore",
 					"AssetRegistry",
 					"AssetTools",
diff --git a/Engine/Plugins/Editor/ContentBrowser/ContentBrowserAssetDataSource/Source/ContentBrowserAssetDataSource/Private/AssetFileContextMenu.cpp b/Engine/Plugins/Editor/ContentBrowser/ContentBrowserAssetDataSource/Source/ContentBrowserAssetDataSource/Private/AssetFileContextMenu.cpp
--- a/Engine/Plugins/Editor/ContentBrowser/ContentBrowserAssetDataSource/Source/ContentBrowserAssetDataSource/Private/AssetFileContextMenu.cpp
+++ b/Engine/Plugins/Editor/ContentBrowser/ContentBrowserAssetDataSource/Source/ContentBrowserAssetDataSource/Private/AssetFileContextMenu.cpp
@@ -1,6 +1,9 @@
 // Copyright Epic Games, Inc. All Rights Reserved.
 
 #include "AssetFileContextMenu.h"
+
+#include "AssetDefinition.h"
+#include "AssetDefinitionRegistry.h"
 #include "Templates/SubclassOf.h"
 #include "Styling/SlateTypes.h"
 #include "Framework/Commands/UIAction.h"
@@ -125,11 +128,6 @@ bool FAssetFileContextMenu::AddImportedAssetMenuOptions(UToolMenu* Menu)
 {
 	if (AreImportedAssetActionsVisible())
 	{
-		TArray<FString> ResolvedFilePaths;
-		TArray<FString> SourceFileLabels;
-		int32 ValidSelectedAssetCount = 0;
-		GetSelectedAssetSourceFilePaths(ResolvedFilePaths, SourceFileLabels, ValidSelectedAssetCount);
-
 		FToolMenuSection& Section = Menu->AddSection("ImportedAssetActions", LOCTEXT("ImportedAssetActionsMenuHeading", "Imported Asset"));
 		Section.InsertPosition = FToolMenuInsert("CommonAssetActions", EToolMenuInsertType::Before);
 		{
@@ -189,6 +187,11 @@ bool FAssetFileContextMenu::AddImportedAssetMenuOptions(UToolMenu* Menu)
 					}
 				}
 			};
+			
+			TArray<FString> ResolvedFilePaths;
+			TArray<FString> SourceFileLabels;
+			int32 ValidSelectedAssetCount = 0;
+			GetSelectedAssetSourceFilePaths(ResolvedFilePaths, SourceFileLabels, ValidSelectedAssetCount);
 
 			//Reimport Menu
 			if (ValidSelectedAssetCount == 1 && SourceFileLabels.Num() > 1)
@@ -264,7 +267,6 @@ bool FAssetFileContextMenu::AddImportedAssetMenuOptions(UToolMenu* Menu)
 
 		return true;
 	}
-	
 
 	return false;
 }
@@ -447,6 +449,18 @@ void FAssetFileContextMenu::MakeAssetActionsSubMenu(UToolMenu* Menu)
 			)
 		);
 
+		// Load
+		Section.AddMenuEntry(
+			"Load",
+			LOCTEXT("Load", "Load"),
+			LOCTEXT("LoadTooltip", "Loads the selected assets into memory without opening any asset editors."),
+			FSlateIcon(),
+			FUIAction(
+				FExecuteAction::CreateSP(this, &FAssetFileContextMenu::ExecuteLoad),
+				FCanExecuteAction::CreateSP(this, &FAssetFileContextMenu::CanExecuteLoad)
+			)
+		);
+
 		// Replace References
 		Section.AddMenuEntry(
 			"ReplaceReferences",
@@ -722,14 +736,7 @@ void FAssetFileContextMenu::MakeAssetLocalizationSubMenu(UToolMenu* Menu)
 		for (const FAssetData& Asset : SelectedAssets)
 		{
 			// Can this type of asset be localized?
-			bool bCanLocalizeAsset = false;
-			{
-				TSharedPtr<IAssetTypeActions> AssetTypeActions = AssetToolsModule.Get().GetAssetTypeActionsForClass(Asset.GetClass()).Pin();
-				if (AssetTypeActions.IsValid())
-				{
-					bCanLocalizeAsset = AssetTypeActions->CanLocalize();
-				}
-			}
+			const bool bCanLocalizeAsset = AssetToolsModule.Get().CanLocalize(Asset.GetClass());
 
 			if (!bCanLocalizeAsset)
 			{
@@ -1229,41 +1236,31 @@ void FAssetFileContextMenu::ExecuteReimportWithNewFile(int32 SourceFileIndex /*=
 	// Ask for a new files and reimport the selected asset
 	check(SelectedAssets.Num() == 1);
 
-	TArray<UObject *> CopyOfSelectedAssets;
-	for (const FAssetData &SelectedAsset : SelectedAssets)
-	{
-		UObject *Asset = SelectedAsset.GetAsset();
-		if (Asset)
-		{
-			CopyOfSelectedAssets.Add(Asset);
-		}
-	}
-
-	TArray<FString> AssetSourcePaths;
-	UClass* ObjectClass = CopyOfSelectedAssets[0]->GetClass();
-	FAssetToolsModule& AssetToolsModule = FModuleManager::Get().LoadModuleChecked<FAssetToolsModule>("AssetTools");
-	const auto AssetTypeActions = AssetToolsModule.Get().GetAssetTypeActionsForClass(ObjectClass);
-	if (AssetTypeActions.IsValid())
-	{
-		AssetTypeActions.Pin()->GetResolvedSourceFilePaths(CopyOfSelectedAssets, AssetSourcePaths);
-	}
-
-	int32 SourceFileIndexToReplace = SourceFileIndex;
-	//Check if the data is valid
-	if (SourceFileIndex == INDEX_NONE)
+	if (const UAssetDefinition* AssetDefinition = UAssetDefinitionRegistry::Get()->GetAssetDefinitionForAsset(SelectedAssets[0]))
 	{
-		if (AssetSourcePaths.Num() > 1)
+		TArray<FAssetSourceFile> OutSourceAssets;
+		FAssetSourceFileArgs SourceFileArgs(SelectedAssets);
+		if (AssetDefinition->GetSourceFiles(SourceFileArgs, OutSourceAssets) == EAssetCommandResult::Handled)
 		{
-			//Ask for a new file for the index 0
-			SourceFileIndexToReplace = 0;
+			int32 SourceFileIndexToReplace = SourceFileIndex;
+			//Check if the data is valid
+			if (SourceFileIndex == INDEX_NONE)
+			{
+				if (OutSourceAssets.Num() > 1)
+				{
+					//Ask for a new file for the index 0
+					SourceFileIndexToReplace = 0;
+				}
+			}
+			else
+			{
+				check(OutSourceAssets.IsValidIndex(SourceFileIndex));
+			}
+			
+			TArray<UObject*> LoadedAssets = SourceFileArgs.LoadObjects<UObject>();
+			FReimportManager::Instance()->ValidateAllSourceFileAndReimport(LoadedAssets, true, SourceFileIndexToReplace, true);
 		}
 	}
-	else
-	{
-		check(AssetSourcePaths.IsValidIndex(SourceFileIndex));
-	}
-
-	FReimportManager::Instance()->ValidateAllSourceFileAndReimport(CopyOfSelectedAssets, true, SourceFileIndexToReplace, true);
 }
 
 void FAssetFileContextMenu::ExecuteFindSourceInExplorer(const TArray<FString> ResolvedFilePaths)
@@ -1284,22 +1281,15 @@ void FAssetFileContextMenu::ExecuteOpenInExternalEditor(const TArray<FString> Re
 	}
 }
 
-void FAssetFileContextMenu::GetSelectedAssetsByClass(TMap<UClass*, TArray<UObject*> >& OutSelectedAssetsByClass) const
+void FAssetFileContextMenu::GetSelectedAssetsByClass(TMap<UClass*, TArray<FAssetData> >& OutSelectedAssetsByClass) const
 {
 	// Sort all selected assets by class
-	for (const auto& SelectedAsset : SelectedAssets)
+	for (const FAssetData& SelectedAsset : SelectedAssets)
 	{
-		auto Asset = SelectedAsset.GetAsset();
-		if (Asset)
+		if (UClass* AssetClass = SelectedAsset.GetClass())
 		{
-			auto AssetClass = Asset->GetClass();
-
-			if ( !OutSelectedAssetsByClass.Contains(AssetClass) )
-			{
-				OutSelectedAssetsByClass.Add(AssetClass);
-			}
-		
-			OutSelectedAssetsByClass[AssetClass].Add(Asset);
+			TArray<FAssetData>& SelectedAssetsForClass = OutSelectedAssetsByClass.FindOrAdd(AssetClass);
+			SelectedAssetsForClass.Add(SelectedAsset);
 		}
 	}
 }
@@ -1308,27 +1298,25 @@ void FAssetFileContextMenu::GetSelectedAssetSourceFilePaths(TArray<FString>& Out
 {
 	OutFilePaths.Empty();
 	OutUniqueSourceFileLabels.Empty();
-	TMap<UClass*, TArray<UObject*> > SelectedAssetsByClass;
+	TMap<UClass*, TArray<FAssetData>> SelectedAssetsByClass;
 	GetSelectedAssetsByClass(SelectedAssetsByClass);
-	FAssetToolsModule& AssetToolsModule = FModuleManager::Get().LoadModuleChecked<FAssetToolsModule>("AssetTools");
+
 	OutValidSelectedAssetCount = 0;
 	// Get the source file paths for the assets of each type
 	for (const auto& AssetsByClassPair : SelectedAssetsByClass)
 	{
-		const auto AssetTypeActions = AssetToolsModule.Get().GetAssetTypeActionsForClass(AssetsByClassPair.Key);
-		if (AssetTypeActions.IsValid())
+		if (const UAssetDefinition* AssetDefinition = UAssetDefinitionRegistry::Get()->GetAssetDefinitionForClass(AssetsByClassPair.Key))
 		{
-			const auto& TypeAssets = AssetsByClassPair.Value;
-			OutValidSelectedAssetCount += TypeAssets.Num();
-			TArray<FString> AssetSourcePaths;
-			AssetTypeActions.Pin()->GetResolvedSourceFilePaths(TypeAssets, AssetSourcePaths);
-			OutFilePaths.Append(AssetSourcePaths);
-
-			TArray<FString> AssetSourceLabels;
-			AssetTypeActions.Pin()->GetSourceFileLabels(TypeAssets, AssetSourceLabels);
-			for (const FString& Label : AssetSourceLabels)
+			TArray<FAssetSourceFile> OutSourceAssets;
+			FAssetSourceFileArgs SourceFileArgs(SelectedAssets);
+			if (AssetDefinition->GetSourceFiles(SourceFileArgs, OutSourceAssets) == EAssetCommandResult::Handled)
 			{
-				OutUniqueSourceFileLabels.AddUnique(Label);
+				OutValidSelectedAssetCount += AssetsByClassPair.Value.Num();
+				for (const FAssetSourceFile& SourceFile : OutSourceAssets)
+				{
+					OutFilePaths.Add(SourceFile.RelativeFilename);
+					OutUniqueSourceFileLabels.AddUnique(SourceFile.DisplayLabelName);
+				}
 			}
 		}
 	}
@@ -1578,6 +1566,23 @@ void FAssetFileContextMenu::ExecuteDiffSelected() const
 	}
 }
 
+bool FAssetFileContextMenu::CanExecuteLoad() const
+{
+	return SelectedAssets.Num() > 0;
+}
+
+void FAssetFileContextMenu::ExecuteLoad()
+{
+	FScopedSlowTask SlowTask(SelectedAssets.Num(), LOCTEXT("LoadingSelectedAssets", "Loading Selected Assets..."));
+	SlowTask.MakeDialogDelayed(1.0f);
+
+	for (const FAssetData& SelectedAsset : SelectedAssets)
+	{
+		SelectedAsset.GetAsset();
+		SlowTask.EnterProgressFrame(1);
+	}
+}
+
 bool FAssetFileContextMenu::CanExecuteReload() const
 {
 	return SelectedAssets.Num() > 0;
diff --git a/Engine/Plugins/Editor/ContentBrowser/ContentBrowserAssetDataSource/Source/ContentBrowserAssetDataSource/Private/ContentBrowserAssetDataCore.cpp b/Engine/Plugins/Editor/ContentBrowser/ContentBrowserAssetDataSource/Source/ContentBrowserAssetDataSource/Private/ContentBrowserAssetDataCore.cpp
--- a/Engine/Plugins/Editor/ContentBrowser/ContentBrowserAssetDataSource/Source/ContentBrowserAssetDataSource/Private/ContentBrowserAssetDataCore.cpp
+++ b/Engine/Plugins/Editor/ContentBrowser/ContentBrowserAssetDataSource/Source/ContentBrowserAssetDataSource/Private/ContentBrowserAssetDataCore.cpp
@@ -1,6 +1,8 @@
 // Copyright Epic Games, Inc. All Rights Reserved.
 
 #include "ContentBrowserAssetDataCore.h"
+
+#include "AssetDefinition.h"
 #include "ContentBrowserDataSource.h"
 #include "IAssetTools.h"
 #include "AssetRegistry/IAssetRegistry.h"
@@ -22,6 +24,7 @@
 #include "Settings/ContentBrowserSettings.h"
 #include "Engine/Level.h"
 #include "Framework/Notifications/NotificationManager.h"
+#include "Misc/WarnIfAssetsLoadedInScope.h"
 #include "Widgets/Notifications/SNotificationList.h"
 
 #define LOCTEXT_NAMESPACE "ContentBrowserAssetDataSource"
@@ -1430,9 +1433,11 @@ bool GetAssetFileItemAttribute(const FContentBrowserAssetFileItemDataPayload& In
 
 		if (InAttributeKey == ContentBrowserItemAttributes::ItemDescription)
 		{
-			if (TSharedPtr<IAssetTypeActions> AssetTypeActions = InAssetPayload.GetAssetTypeActions())
+			if (const UAssetDefinition* AssetDefinition = InAssetPayload.GetAssetDefinition())
 			{
-				const FText AssetDescription = AssetTypeActions->GetAssetDescription(InAssetPayload.GetAssetData());
+				FWarnIfAssetsLoadedInScope WarnIfAssetLoadedInScope({ InAssetPayload.GetAssetData() });
+				
+				const FText AssetDescription = AssetDefinition->GetAssetDescription(InAssetPayload.GetAssetData());
 				if (!AssetDescription.IsEmpty())
 				{
 					OutAttributeValue.SetValue(AssetDescription);
diff --git a/Engine/Plugins/Editor/ContentBrowser/ContentBrowserAssetDataSource/Source/ContentBrowserAssetDataSource/Private/ContentBrowserAssetDataPayload.cpp b/Engine/Plugins/Editor/ContentBrowser/ContentBrowserAssetDataSource/Source/ContentBrowserAssetDataSource/Private/ContentBrowserAssetDataPayload.cpp
--- a/Engine/Plugins/Editor/ContentBrowser/ContentBrowserAssetDataSource/Source/ContentBrowserAssetDataSource/Private/ContentBrowserAssetDataPayload.cpp
+++ b/Engine/Plugins/Editor/ContentBrowser/ContentBrowserAssetDataSource/Source/ContentBrowserAssetDataSource/Private/ContentBrowserAssetDataPayload.cpp
@@ -4,6 +4,8 @@
 #include "Modules/ModuleManager.h"
 #include "AssetToolsModule.h"
 #include "AssetThumbnail.h"
+#include "AssetDefinition.h"
+#include "AssetDefinitionRegistry.h"
 #include "Engine/Texture2D.h"
 #include "Materials/Material.h"
 
@@ -109,6 +111,16 @@ TSharedPtr<IAssetTypeActions> FContentBrowserAssetFileItemDataPayload::GetAssetT
 	return CachedAssetTypeActionsPtr.Pin();
 }
 
+const UAssetDefinition* FContentBrowserAssetFileItemDataPayload::GetAssetDefinition() const
+{
+	if (!bHasCachedAssetDefinitionPtr)
+	{
+		CachedAssetDefinitionPtr = UAssetDefinitionRegistry::Get()->GetAssetDefinitionForAsset(AssetData);
+		bHasCachedAssetDefinitionPtr = true;
+	}
+	return CachedAssetDefinitionPtr.Get();
+}
+
 const FString& FContentBrowserAssetFileItemDataPayload::GetFilename() const
 {
 	if (!bHasCachedFilename)
diff --git a/Engine/Plugins/Editor/ContentBrowser/ContentBrowserAssetDataSource/Source/ContentBrowserAssetDataSource/Private/NewAssetContextMenu.cpp b/Engine/Plugins/Editor/ContentBrowser/ContentBrowserAssetDataSource/Source/ContentBrowserAssetDataSource/Private/NewAssetContextMenu.cpp
--- a/Engine/Plugins/Editor/ContentBrowser/ContentBrowserAssetDataSource/Source/ContentBrowserAssetDataSource/Private/NewAssetContextMenu.cpp
+++ b/Engine/Plugins/Editor/ContentBrowser/ContentBrowserAssetDataSource/Source/ContentBrowserAssetDataSource/Private/NewAssetContextMenu.cpp
@@ -69,8 +69,10 @@ struct FCategorySubMenuItem
 	}
 };
 
-TArray<FFactoryItem> FindFactoriesInCategory(EAssetTypeCategories::Type AssetTypeCategory)
+TArray<FFactoryItem> FindFactoriesInCategory(EAssetTypeCategories::Type AssetTypeCategory, bool FindFirstOnly)
 {
+    QUICK_SCOPE_CYCLE_COUNTER(FindFactoriesInCategory);
+	
 	TArray<FFactoryItem> FactoriesInThisCategory;
 
 	static const FName NAME_AssetTools = "AssetTools";
@@ -78,10 +80,17 @@ TArray<FFactoryItem> FindFactoriesInCategory(EAssetTypeCategories::Type AssetTyp
 	TArray<UFactory*> Factories = AssetTools.GetNewAssetFactories();
 	for (UFactory* Factory : Factories)
 	{
-		uint32 FactoryCategories = Factory->GetMenuCategories();
+		QUICK_SCOPE_CYCLE_COUNTER(GetMenuCategories);
+	
+		const uint32 FactoryCategories = Factory->GetMenuCategories();
 		if (FactoryCategories & AssetTypeCategory)
 		{
 			FactoriesInThisCategory.Emplace(Factory, Factory->GetDisplayName());
+
+			if (FindFirstOnly)
+			{
+				return FactoriesInThisCategory;
+			}
 		}
 	}
 
@@ -199,6 +208,8 @@ void FNewAssetContextMenu::MakeContextMenu(
 	const FOnNewAssetRequested& InOnNewAssetRequested
 	)
 {
+	QUICK_SCOPE_CYCLE_COUNTER(STAT_ContentBrowser_MakeNewAssetContextMenu);
+	
 	if (InSelectedAssetPaths.Num() == 0)
 	{
 		return;
@@ -232,6 +243,8 @@ void FNewAssetContextMenu::MakeContextMenu(
 	// Import
 	if (InOnImportAssetRequested.IsBound() && !FirstSelectedPath.IsNone())
 	{
+		QUICK_SCOPE_CYCLE_COUNTER(STAT_ContentBrowser_ImportSection);
+	
 		{
 			FToolMenuSection& Section = Menu->FindOrAddSection("ContentBrowserGetContent");
 			Section.AddMenuEntry(
@@ -252,6 +265,8 @@ void FNewAssetContextMenu::MakeContextMenu(
 	{
 		// Add Basic Asset
 		{
+			QUICK_SCOPE_CYCLE_COUNTER(STAT_ContentBrowser_BasicSection);
+		
 			FToolMenuSection& Section = Menu->AddSection("ContentBrowserNewBasicAsset", LOCTEXT("CreateBasicAssetsMenuHeading", "Create Basic Asset"));
 			CreateNewAssetMenuCategory(
 				Menu,
@@ -265,6 +280,8 @@ void FNewAssetContextMenu::MakeContextMenu(
 
 		// Add Advanced Asset
 		{
+			QUICK_SCOPE_CYCLE_COUNTER(STAT_ContentBrowser_AdvancedSection);
+		
 			FToolMenuSection& Section = Menu->AddSection("ContentBrowserNewAdvancedAsset", LOCTEXT("CreateAdvancedAssetsMenuHeading", "Create Advanced Asset"));
 
 			TArray<FAdvancedAssetCategory> AdvancedAssetCategories;
@@ -275,7 +292,8 @@ void FNewAssetContextMenu::MakeContextMenu(
 
 			for (const FAdvancedAssetCategory& AdvancedAssetCategory : AdvancedAssetCategories)
 			{
-				TArray<FFactoryItem> Factories = FindFactoriesInCategory(AdvancedAssetCategory.CategoryType);
+				const bool FindFirstOnly = true;
+				TArray<FFactoryItem> Factories = FindFactoriesInCategory(AdvancedAssetCategory.CategoryType, FindFirstOnly);
 				if (Factories.Num() > 0)
 				{
 					Section.AddSubMenu(
@@ -305,7 +323,8 @@ void FNewAssetContextMenu::MakeContextMenu(
 void FNewAssetContextMenu::CreateNewAssetMenuCategory(UToolMenu* Menu, FName SectionName, EAssetTypeCategories::Type AssetTypeCategory, FName InPath, FOnNewAssetRequested InOnNewAssetRequested, FCanExecuteAction InCanExecuteAction)
 {
 	// Find UFactory classes that can create new objects in this category.
-	TArray<FFactoryItem> FactoriesInThisCategory = FindFactoriesInCategory(AssetTypeCategory);
+	const bool FindFirstOnly = false;
+	TArray<FFactoryItem> FactoriesInThisCategory = FindFactoriesInCategory(AssetTypeCategory, FindFirstOnly);
 	if (FactoriesInThisCategory.Num() == 0)
 	{
 		return;
@@ -335,11 +354,14 @@ void FNewAssetContextMenu::CreateNewAssetMenuCategory(UToolMenu* Menu, FName Sec
 		SubMenu->Factories.Add(Item);
 	}
 	ParentMenuData->SortSubMenus();
+	
 	CreateNewAssetMenus(Menu, SectionName, ParentMenuData, InPath, InOnNewAssetRequested, InCanExecuteAction);
 }
 
 void FNewAssetContextMenu::CreateNewAssetMenus(UToolMenu* Menu, FName SectionName, TSharedPtr<FCategorySubMenuItem> SubMenuData, FName InPath, FOnNewAssetRequested InOnNewAssetRequested, FCanExecuteAction InCanExecuteAction)
 {
+	QUICK_SCOPE_CYCLE_COUNTER(CreateNewAssetMenus);
+
 	FToolMenuSection& Section = Menu->FindOrAddSection(SectionName);
 	for (const FFactoryItem& FactoryItem : SubMenuData->Factories)
 	{
diff --git a/Engine/Plugins/Editor/ContentBrowser/ContentBrowserAssetDataSource/Source/ContentBrowserAssetDataSource/Public/AssetFileContextMenu.h b/Engine/Plugins/Editor/ContentBrowser/ContentBrowserAssetDataSource/Source/ContentBrowserAssetDataSource/Public/AssetFileContextMenu.h
--- a/Engine/Plugins/Editor/ContentBrowser/ContentBrowserAssetDataSource/Source/ContentBrowserAssetDataSource/Public/AssetFileContextMenu.h
+++ b/Engine/Plugins/Editor/ContentBrowser/ContentBrowserAssetDataSource/Source/ContentBrowserAssetDataSource/Public/AssetFileContextMenu.h
@@ -38,7 +38,7 @@ private:
 
 private:
 	/** Helper to load selected assets and sort them by UClass */
-	void GetSelectedAssetsByClass(TMap<UClass*, TArray<UObject*> >& OutSelectedAssetsByClass) const;
+	void GetSelectedAssetsByClass(TMap<UClass*, TArray<FAssetData> >& OutSelectedAssetsByClass) const;
 
 	/** Helper to collect resolved filepaths for all selected assets */
 	void GetSelectedAssetSourceFilePaths(TArray<FString>& OutFilePaths, TArray<FString>& OutUniqueSourceFileLabels, int32 &OutValidSelectedAssetCount) const;
@@ -63,6 +63,12 @@ private:
 	/** Handler for OpenInExternalEditor */
 	void ExecuteOpenInExternalEditor(const TArray<FString> ResolvedFilePaths);
 
+	/** Handler to check to see if a load command is allowed */
+	bool CanExecuteLoad() const;
+
+	/** Handler for Load */
+	void ExecuteLoad();
+
 	/** Handler to check to see if a reload command is allowed */
 	bool CanExecuteReload() const;
 
diff --git a/Engine/Plugins/Editor/ContentBrowser/ContentBrowserAssetDataSource/Source/ContentBrowserAssetDataSource/Public/ContentBrowserAssetDataPayload.h b/Engine/Plugins/Editor/ContentBrowser/ContentBrowserAssetDataSource/Source/ContentBrowserAssetDataSource/Public/ContentBrowserAssetDataPayload.h
--- a/Engine/Plugins/Editor/ContentBrowser/ContentBrowserAssetDataSource/Source/ContentBrowserAssetDataSource/Public/ContentBrowserAssetDataPayload.h
+++ b/Engine/Plugins/Editor/ContentBrowser/ContentBrowserAssetDataSource/Source/ContentBrowserAssetDataSource/Public/ContentBrowserAssetDataPayload.h
@@ -8,6 +8,7 @@
 #include "UObject/GCObject.h"
 #include "Containers/Set.h"
 
+class UAssetDefinition;
 class IAssetTypeActions;
 class FAssetThumbnail;
 class UFactory;
@@ -56,6 +57,8 @@ public:
 	UObject* LoadAsset(TSet<FName> LoadTags = {}) const;
 
 	TSharedPtr<IAssetTypeActions> GetAssetTypeActions() const;
+	
+	const UAssetDefinition* GetAssetDefinition() const;
 
 	const FString& GetFilename() const;
 
@@ -73,6 +76,9 @@ private:
 	mutable bool bHasCachedAssetTypeActionsPtr = false;
 	mutable TWeakPtr<IAssetTypeActions> CachedAssetTypeActionsPtr;
 
+	mutable bool bHasCachedAssetDefinitionPtr = false;
+	mutable TWeakObjectPtr<const UAssetDefinition> CachedAssetDefinitionPtr;
+
 	mutable bool bHasCachedFilename = false;
 	mutable FString CachedFilename;
 
diff --git a/Engine/Plugins/Editor/DataValidation/Source/DataValidation/Private/EditorValidator_Localization.cpp b/Engine/Plugins/Editor/DataValidation/Source/DataValidation/Private/EditorValidator_Localization.cpp
--- a/Engine/Plugins/Editor/DataValidation/Source/DataValidation/Private/EditorValidator_Localization.cpp
+++ b/Engine/Plugins/Editor/DataValidation/Source/DataValidation/Private/EditorValidator_Localization.cpp
@@ -48,13 +48,10 @@ EDataValidationResult UEditorValidator_Localization::ValidateLoadedAsset_Impleme
 			if (SourceAssetData.IsValid())
 			{
 				// Can this type of asset be localized?
-				if (TSharedPtr<IAssetTypeActions> AssetTypeActions = AssetTools.GetAssetTypeActionsForClass(InAsset->GetClass()).Pin())
+				if (!AssetTools.CanLocalize(InAsset->GetClass()))
 				{
-					if (!AssetTypeActions->CanLocalize())
-					{
-						AssetFails(InAsset, FText::Format(LOCTEXT("LocalizationError_LocalizedAssetTypeCannotBeLocalized", "Localized asset is of type '{0}', which is not a type that can be localized!"), FText::FromString(InAsset->GetClass()->GetName())), ValidationErrors);
-						return EDataValidationResult::Invalid;
-					}
+					AssetFails(InAsset, FText::Format(LOCTEXT("LocalizationError_LocalizedAssetTypeCannotBeLocalized", "Localized asset is of type '{0}', which is not a type that can be localized!"), FText::FromString(InAsset->GetClass()->GetName())), ValidationErrors);
+					return EDataValidationResult::Invalid;
 				}
 
 				// Is the source asset a redirector?
@@ -94,13 +91,10 @@ EDataValidationResult UEditorValidator_Localization::ValidateLoadedAsset_Impleme
 						if (LocalizedAssetData.IsValid())
 						{
 							// Can this type of asset be localized?
-							if (TSharedPtr<IAssetTypeActions> AssetTypeActions = AssetTools.GetAssetTypeActionsForClass(InAsset->GetClass()).Pin())
+							if (!AssetTools.CanLocalize(InAsset->GetClass()))
 							{
-								if (!AssetTypeActions->CanLocalize())
-								{
-									AssetFails(InAsset, FText::Format(LOCTEXT("LocalizationError_SourceAssetTypeCannotBeLocalized", "Source asset has a localized asset for '{0}', but is of type '{1}' which is not a type that can be localized!"), FText::FromString(CultureName), FText::FromString(InAsset->GetClass()->GetName())), ValidationErrors);
-									return EDataValidationResult::Invalid;
-								}
+								AssetFails(InAsset, FText::Format(LOCTEXT("LocalizationError_SourceAssetTypeCannotBeLocalized", "Source asset has a localized asset for '{0}', but is of type '{1}' which is not a type that can be localized!"), FText::FromString(CultureName), FText::FromString(InAsset->GetClass()->GetName())), ValidationErrors);
+								return EDataValidationResult::Invalid;
 							}
 
 							// Is the source asset a redirector?
diff --git a/Engine/Plugins/Editor/EngineAssetDefinitions/EngineAssetDefinitions.uplugin b/Engine/Plugins/Editor/EngineAssetDefinitions/EngineAssetDefinitions.uplugin
new file mode 100644
--- /dev/null
+++ b/Engine/Plugins/Editor/EngineAssetDefinitions/EngineAssetDefinitions.uplugin
@@ -0,0 +1,25 @@
+{
+	"FileVersion": 3,
+	"Version": 1,
+	"VersionName": "1.0",
+	"FriendlyName": "Engine Asset Definitions",
+	"Description": "",
+	"Category": "Other",
+	"CreatedBy" : "Epic Games, Inc.",
+	"CreatedByURL" : "http://epicgames.com",
+	"DocsURL": "",
+	"MarketplaceURL": "",
+	"SupportURL": "",
+	"EnabledByDefault" : true,
+	"CanContainContent": false,
+	"IsBetaVersion": false,
+	"IsExperimentalVersion": false,
+	"Installed": false,
+	"Modules": [
+		{
+			"Name": "EngineAssetDefinitions",
+			"Type": "Editor",
+			"LoadingPhase": "PostEngineInit"
+		}
+	]
+}
\ No newline at end of file
diff --git a/Engine/Plugins/Editor/EngineAssetDefinitions/Source/EngineAssetDefinitions.Build.cs b/Engine/Plugins/Editor/EngineAssetDefinitions/Source/EngineAssetDefinitions.Build.cs
new file mode 100644
--- /dev/null
+++ b/Engine/Plugins/Editor/EngineAssetDefinitions/Source/EngineAssetDefinitions.Build.cs
@@ -0,0 +1,32 @@
+// Copyright Epic Games, Inc. All Rights Reserved.
+
+using UnrealBuildTool;
+
+public class EngineAssetDefinitions : ModuleRules
+{
+	public EngineAssetDefinitions(ReadOnlyTargetRules Target) : base(Target)
+	{
+		PublicDependencyModuleNames.AddRange(
+			new string[]
+			{
+				"Core",
+			}
+		);			
+		
+		PrivateDependencyModuleNames.AddRange(
+			new string[]
+			{
+				"ContentBrowser",
+				"AssetDefinition",
+				"UnrealEd",
+				"ToolMenus",
+				"CoreUObject",
+				"Engine",
+				"AssetTools",
+				"SlateCore",
+				"Slate",
+				"Kismet"
+			}
+		);
+	}
+}
diff --git a/Engine/Plugins/Editor/EngineAssetDefinitions/Source/Private/AssetDefinition_DataAsset.cpp b/Engine/Plugins/Editor/EngineAssetDefinitions/Source/Private/AssetDefinition_DataAsset.cpp
new file mode 100644
--- /dev/null
+++ b/Engine/Plugins/Editor/EngineAssetDefinitions/Source/Private/AssetDefinition_DataAsset.cpp
@@ -0,0 +1,221 @@
+// Copyright Epic Games, Inc. All Rights Reserved.
+
+#include "AssetDefinition_DataAsset.h"
+
+#include "AssetToolsModule.h"
+#include "ClassViewerFilter.h"
+#include "ClassViewerModule.h"
+#include "ContentBrowserMenuContexts.h"
+#include "ObjectTools.h"
+#include "SDetailsDiff.h"
+#include "ToolMenus.h"
+#include "Kismet2/SClassPickerDialog.h"
+
+#define LOCTEXT_NAMESPACE "UAssetDefinition_DataAsset"
+
+namespace MenuExtension_DataAsset
+{
+	class FNewNodeClassFilter : public IClassViewerFilter
+	{
+	public:
+		FNewNodeClassFilter(UClass* InBaseClass)
+			: BaseClass(InBaseClass)
+		{
+		}
+
+		virtual bool IsClassAllowed(const FClassViewerInitializationOptions& InInitOptions, const UClass* InClass, TSharedRef< FClassViewerFilterFuncs > InFilterFuncs) override
+		{
+			if (InClass != nullptr)
+			{
+				return InClass->IsChildOf(BaseClass);
+			}
+			return false;
+		}
+
+		virtual bool IsUnloadedClassAllowed(const FClassViewerInitializationOptions& InInitOptions, const TSharedRef< const IUnloadedBlueprintData > InUnloadedClassData, TSharedRef< FClassViewerFilterFuncs > InFilterFuncs) override
+		{
+			return InUnloadedClassData->IsChildOf(BaseClass);
+		}
+
+	private:
+		UClass* BaseClass;
+	};
+
+	void ExecuteChangeDataAssetClass(const FToolMenuContext& MenuContext)
+	{
+		if (const UContentBrowserAssetContextMenuContext* Context = UContentBrowserAssetContextMenuContext::FindContextWithAssets(MenuContext))
+		{
+			TArray<UDataAsset*> DataAssets = Context->LoadSelectedObjects<UDataAsset>();
+
+			const FText TitleText = LOCTEXT("DataAsset_PickNewDataAssetClass", "Pick New DataAsset Class");
+			FClassViewerInitializationOptions Options;
+			Options.ClassFilters.Add(MakeShared<FNewNodeClassFilter>(UDataAsset::StaticClass()));
+			UClass* OutNewDataAssetClass = nullptr;
+			const bool bPressedOk = SClassPickerDialog::PickClass(TitleText, Options, OutNewDataAssetClass, UDataAsset::StaticClass());
+
+			if (bPressedOk && OutNewDataAssetClass != nullptr)
+			{
+				for (TWeakObjectPtr<UDataAsset> DataAssetPtr : DataAssets)
+				{
+					if (UDataAsset* OldDataAsset = DataAssetPtr.Get())
+					{
+						if (OldDataAsset && OldDataAsset->IsValidLowLevel())
+						{
+							FName ObjectName = OldDataAsset->GetFName();
+							UObject* Outer = OldDataAsset->GetOuter();
+							OldDataAsset->Rename(nullptr, GetTransientPackage(), REN_DoNotDirty | REN_DontCreateRedirectors);
+
+							UObject* NewDataAsset = NewObject<UObject>(Outer, OutNewDataAssetClass, ObjectName, OldDataAsset->GetFlags());
+
+							// Migrate Data
+							{
+								UEngine::FCopyPropertiesForUnrelatedObjectsParams CopyOptions;
+								CopyOptions.bNotifyObjectReplacement = true;
+								UEngine::CopyPropertiesForUnrelatedObjects(OldDataAsset, NewDataAsset, CopyOptions);
+							}
+
+							NewDataAsset->MarkPackageDirty();
+
+							// Consolidate or "Replace" the old object with the new object for any living references.
+							bool bShowDeleteConfirmation = false;
+							TArray<UObject*> OldDataAssetArray({ (UObject*)OldDataAsset });
+							ObjectTools::ConsolidateObjects(NewDataAsset, OldDataAssetArray, bShowDeleteConfirmation);
+						}
+					}
+				}
+			}
+		}
+	}
+	
+	static FDelayedAutoRegisterHelper DelayedAutoRegister(EDelayedRegisterRunPhase::EndOfEngineInit, []{ 
+		UToolMenus::RegisterStartupCallback(FSimpleMulticastDelegate::FDelegate::CreateLambda([]()
+		{
+			FToolMenuOwnerScoped OwnerScoped(UE_MODULE_NAME);
+			UToolMenu* Menu = UE::ContentBrowser::ExtendToolMenu_AssetContextMenu(UDataAsset::StaticClass());
+	        
+			FToolMenuSection& Section = Menu->FindOrAddSection("GetAssetActions");
+			Section.AddDynamicEntry("GetAssetActions_DataAsset", FNewToolMenuSectionDelegate::CreateStatic([](FToolMenuSection& InSection)
+			{
+				{
+					const TAttribute<FText> Label = LOCTEXT("DataAsset_ChangeClass", "Convert to Different DataAsset Type");
+					const TAttribute<FText> ToolTip = LOCTEXT("DataAsset_ChangeClassTip", "Change the class these Data Assets are subclassed from.");
+					const FSlateIcon Icon = FSlateIcon(FAppStyle::GetAppStyleSetName(), "ClassIcon.DataAsset");
+					const FToolMenuExecuteAction UIAction = FToolMenuExecuteAction::CreateStatic(&ExecuteChangeDataAssetClass);
+	    
+					InSection.AddMenuEntry("DataAsset_ChangeClass", Label, ToolTip, Icon, UIAction);
+				}
+			}));
+		}));
+	});
+
+}
+
+
+FText UAssetDefinition_DataAsset::GetAssetDisplayName(const FAssetData& AssetData) const
+{
+	static const FName NAME_RowStructure(TEXT("RowStructure"));
+	
+	if (AssetData.IsValid())
+	{
+		const FAssetDataTagMapSharedView::FFindTagResult RowStructureTag = AssetData.TagsAndValues.FindTag(NAME_RowStructure);
+		if (RowStructureTag.IsSet())
+		{
+			if (UScriptStruct* FoundStruct = UClass::TryFindTypeSlow<UScriptStruct>(RowStructureTag.GetValue(), EFindFirstObjectOptions::ExactClass))
+			{
+				return FText::Format(LOCTEXT("DataTableWithRowType", "Data Table ({0})"), FoundStruct->GetDisplayNameText());
+			}
+		}
+	}
+
+	return FText::GetEmpty();
+}
+
+EAssetCommandResult UAssetDefinition_DataAsset::PerformAssetDiff(const FAssetDiffArgs& DiffArgs) const
+{
+	// sometimes we're comparing different revisions of one single asset (other 
+	// times we're comparing two completely separate assets altogether)
+	const bool bIsSingleAsset = (DiffArgs.OldAsset->GetFName() == DiffArgs.NewAsset->GetFName());
+	static const FText BasicWindowTitle = LOCTEXT("NamelessDataAssetDiff", "DataAsset Diff");
+
+	const FText WindowTitle = !bIsSingleAsset ? BasicWindowTitle : FText::Format(LOCTEXT("DataAsset Diff", "{0} - DataAsset Diff"), FText::FromString(DiffArgs.NewAsset->GetName()));
+
+	SDetailsDiff::CreateDiffWindow(WindowTitle, DiffArgs.OldAsset, DiffArgs.NewAsset, DiffArgs.OldRevision, DiffArgs.NewRevision);
+	
+	return EAssetCommandResult::Handled;
+}
+
+void UAssetDefinition_DataAsset::OnRegistered()
+{
+	UToolMenus::RegisterStartupCallback(FSimpleMulticastDelegate::FDelegate::CreateUObject(this, &ThisClass::RegisterMenus));
+}
+
+void UAssetDefinition_DataAsset::RegisterMenus()
+{
+	return;
+	//FToolMenuOwnerScoped MenuOwner(GetClass()->GetOutermost()->GetLoadedPath());
+
+	UToolMenu* Menu = UE::ContentBrowser::ExtendToolMenu_AssetContextMenu(GetAssetClass());
+	
+	FToolMenuSection& Section = Menu->FindOrAddSection("GetAssetActions");
+	Section.AddDynamicEntry("GetAssetActions_DataAsset", FNewToolMenuSectionDelegate::CreateWeakLambda(this, [this](FToolMenuSection& InSection)
+	{
+		{
+			const TAttribute<FText> Label = LOCTEXT("DataAsset_ChangeClass", "Convert to Different DataAsset Type");
+			const TAttribute<FText> ToolTip = LOCTEXT("DataAsset_ChangeClassTip", "Change the class these Data Assets are subclassed from.");
+			const FSlateIcon Icon = FSlateIcon(FAppStyle::GetAppStyleSetName(), "ClassIcon.DataAsset");
+			const FToolMenuExecuteAction UIAction = FToolMenuExecuteAction::CreateUObject(this, &ThisClass::ExecuteChangeDataAssetClass);
+
+			InSection.AddMenuEntry("DataAsset_ChangeClass", Label, ToolTip, Icon, UIAction);
+		}
+	}));
+}
+
+void UAssetDefinition_DataAsset::ExecuteChangeDataAssetClass(const FToolMenuContext& MenuContext)
+{
+	/*
+	if (const UContentBrowserAssetContextMenuContext* Context = UContentBrowserAssetContextMenuContext::FindContextWithAssets(MenuContext))
+	{
+		TArray<UDataAsset*> DataAssets = Context->LoadSelectedObjects<UDataAsset>();
+
+		const FText TitleText = LOCTEXT("DataAsset_PickNewDataAssetClass", "Pick New DataAsset Class");
+		FClassViewerInitializationOptions Options;
+		Options.ClassFilters.Add(MakeShared<FNewNodeClassFilter>(UDataAsset::StaticClass()));
+		UClass* OutNewDataAssetClass = nullptr;
+		const bool bPressedOk = SClassPickerDialog::PickClass(TitleText, Options, OutNewDataAssetClass, UDataAsset::StaticClass());
+
+		if (bPressedOk && OutNewDataAssetClass != nullptr)
+		{
+			for (TWeakObjectPtr<UDataAsset> DataAssetPtr : DataAssets)
+			{
+				if (UDataAsset* OldDataAsset = DataAssetPtr.Get())
+				{
+					if (OldDataAsset && OldDataAsset->IsValidLowLevel())
+					{
+						FName ObjectName = OldDataAsset->GetFName();
+						UObject* Outer = OldDataAsset->GetOuter();
+						OldDataAsset->Rename(nullptr, GetTransientPackage(), REN_DoNotDirty | REN_DontCreateRedirectors);
+
+						UObject* NewDataAsset = NewObject<UObject>(Outer, OutNewDataAssetClass, ObjectName, OldDataAsset->GetFlags());
+
+						// Migrate Data
+						{
+							UEngine::FCopyPropertiesForUnrelatedObjectsParams CopyOptions;
+							CopyOptions.bNotifyObjectReplacement = true;
+							UEngine::CopyPropertiesForUnrelatedObjects(OldDataAsset, NewDataAsset, CopyOptions);
+						}
+
+						NewDataAsset->MarkPackageDirty();
+
+						// Consolidate or "Replace" the old object with the new object for any living references.
+						bool bShowDeleteConfirmation = false;
+						TArray<UObject*> OldDataAssetArray({ (UObject*)OldDataAsset });
+						ObjectTools::ConsolidateObjects(NewDataAsset, OldDataAssetArray, bShowDeleteConfirmation);
+					}
+				}
+			}
+		}
+	}
+	*/
+}
+
+#undef LOCTEXT_NAMESPACE
diff --git a/Engine/Plugins/Editor/EngineAssetDefinitions/Source/Private/AssetDefinition_DataAsset.h b/Engine/Plugins/Editor/EngineAssetDefinitions/Source/Private/AssetDefinition_DataAsset.h
new file mode 100644
--- /dev/null
+++ b/Engine/Plugins/Editor/EngineAssetDefinitions/Source/Private/AssetDefinition_DataAsset.h
@@ -0,0 +1,34 @@
+// Copyright Epic Games, Inc. All Rights Reserved.
+
+#pragma once
+
+#include "CoreMinimal.h"
+#include "AssetDefinitionDefault.h"
+
+#include "AssetDefinition_DataAsset.generated.h"
+
+struct FToolMenuContext;
+
+UCLASS()
+class UAssetDefinition_DataAsset : public UAssetDefinitionDefault
+{
+	GENERATED_BODY()
+
+public:
+	// UAssetDefinition Implementation
+	virtual FText GetAssetDisplayName() const override { return NSLOCTEXT("AssetDefinition", "DataAsset", "Data Asset"); }
+	virtual FText GetAssetDisplayName(const FAssetData& AssetData) const override;
+	virtual FLinearColor GetAssetColor() const override { return FLinearColor(FColor(201, 29, 85)); }
+	virtual TSoftClassPtr<UObject> GetAssetClass() const override { return UDataAsset::StaticClass(); }
+	
+	virtual EAssetCommandResult PerformAssetDiff(const FAssetDiffArgs& DiffArgs) const override;
+
+protected:
+	virtual void OnRegistered() override;
+
+private:
+	void RegisterMenus();
+	
+	/** Handler for when FindTarget is selected */
+	void ExecuteChangeDataAssetClass(const FToolMenuContext& MenuContext);
+};
diff --git a/Engine/Plugins/Editor/EngineAssetDefinitions/Source/Private/AssetDefinition_Redirector.cpp b/Engine/Plugins/Editor/EngineAssetDefinitions/Source/Private/AssetDefinition_Redirector.cpp
new file mode 100644
--- /dev/null
+++ b/Engine/Plugins/Editor/EngineAssetDefinitions/Source/Private/AssetDefinition_Redirector.cpp
@@ -0,0 +1,116 @@
+// Copyright Epic Games, Inc. All Rights Reserved.
+
+#include "AssetDefinition_Redirector.h"
+
+#include "AssetToolsModule.h"
+#include "ContentBrowserMenuContexts.h"
+#include "ToolMenus.h"
+
+#define LOCTEXT_NAMESPACE "UAssetDefinition_Redirector"
+
+void UAssetDefinition_Redirector::OnRegistered()
+{
+	UToolMenus::RegisterStartupCallback(FSimpleMulticastDelegate::FDelegate::CreateUObject(this, &ThisClass::RegisterMenus));
+}
+
+void UAssetDefinition_Redirector::RegisterMenus()
+{
+	//FToolMenuOwnerScoped MenuOwner(GetClass()->GetOutermost()->GetLoadedPath());
+
+	UToolMenu* Menu = UE::ContentBrowser::ExtendToolMenu_AssetContextMenu(GetAssetClass());
+	
+	FToolMenuSection& Section = Menu->FindOrAddSection("GetAssetActions");
+	Section.AddDynamicEntry("GetAssetActions_ObjectRedirector", FNewToolMenuSectionDelegate::CreateWeakLambda(this, [this](FToolMenuSection& InSection)
+	{
+		{
+			const TAttribute<FText> Label = LOCTEXT("Redirector_FindTarget", "Find Target");
+			const TAttribute<FText> ToolTip = LOCTEXT("Redirector_FindTargetTooltip", "Finds the asset that this redirector targets in the asset tree.");
+			const FSlateIcon Icon = FSlateIcon(FAppStyle::GetAppStyleSetName(), "ClassIcon.ObjectRedirector");
+			const FToolMenuExecuteAction UIAction = FToolMenuExecuteAction::CreateUObject(this, &ThisClass::ExecuteFindTarget);
+
+			InSection.AddMenuEntry("Redirector_FindTarget", Label, ToolTip, Icon, UIAction);
+		}
+		{
+			const TAttribute<FText> Label = LOCTEXT("Redirector_FixUp", "Fix Up");
+			const TAttribute<FText> ToolTip = LOCTEXT("Redirector_FixUpTooltip", "Finds referencers to selected redirectors and resaves them if possible, then deletes any redirectors that had all their referencers fixed.");
+			const FSlateIcon Icon = FSlateIcon(FAppStyle::GetAppStyleSetName(), "ClassIcon.ObjectRedirector");
+			const FToolMenuExecuteAction UIAction = FToolMenuExecuteAction::CreateUObject(this, &ThisClass::ExecuteFixUp, true);
+
+			InSection.AddMenuEntry("Redirector_FixUp", Label, ToolTip, Icon, UIAction);
+		}
+		{
+			const TAttribute<FText> Label = LOCTEXT("Redirector_FixUp_KeepingRedirector", "Fix Up (Keep Redirector)");
+			const TAttribute<FText> ToolTip = LOCTEXT("Redirector_FixUp_KeepingRedirectorTooltip", "Finds referencers to selected redirectors and resaves them if possible.");
+			const FSlateIcon Icon = FSlateIcon(FAppStyle::GetAppStyleSetName(), "ClassIcon.MyAsset");
+			const FToolMenuExecuteAction UIAction = FToolMenuExecuteAction::CreateUObject(this, &ThisClass::ExecuteFixUp, false);
+
+			InSection.AddMenuEntry("Redirector_FixUp_KeepingRedirector", Label, ToolTip, Icon, UIAction);
+		}
+	}));
+}
+
+EAssetCommandResult UAssetDefinition_Redirector::ActivateAssets(const FAssetActivateArgs& ActivateArgs) const
+{
+	if ( ActivateArgs.ActivationMethod == EAssetActivationMethod::DoubleClicked || ActivateArgs.ActivationMethod == EAssetActivationMethod::Opened )
+	{
+		// Sync to the target instead of opening an editor when double clicked
+		TArray<UObjectRedirector*> Redirectors = ActivateArgs.LoadObjects<UObjectRedirector>();
+		if ( Redirectors.Num() > 0 )
+		{
+			FindTargets(Redirectors);
+			return EAssetCommandResult::Handled;
+		}
+	}
+
+	return EAssetCommandResult::Unhandled;
+}
+
+void UAssetDefinition_Redirector::ExecuteFindTarget(const FToolMenuContext& MenuContext)
+{
+	if (const UContentBrowserAssetContextMenuContext* Context = UContentBrowserAssetContextMenuContext::FindContextWithAssets(MenuContext))
+	{
+		TArray<UObjectRedirector*> Redirectors = Context->LoadSelectedObjects<UObjectRedirector>();
+		FindTargets(Redirectors);
+	}
+}
+
+void UAssetDefinition_Redirector::ExecuteFixUp(const FToolMenuContext& MenuContext, bool bDeleteAssets)
+{
+	// This will fix references to selected redirectors, except in the following cases:
+	// Redirectors referenced by unloaded maps will not be fixed up, but any references to it that can be fixed up will
+	// Redirectors referenced by code will not be completely fixed up
+	// Redirectors that are not at head revision or checked out by another user will not be completely fixed up
+	// Redirectors whose referencers are not at head revision, are checked out by another user, or are refused to be checked out will not be completely fixed up.
+
+	if (const UContentBrowserAssetContextMenuContext* Context = UContentBrowserAssetContextMenuContext::FindContextWithAssets(MenuContext))
+	{
+		TArray<UObjectRedirector*> Redirectors = Context->LoadSelectedObjects<UObjectRedirector>();
+
+		if (Redirectors.Num() > 0)
+		{
+			IAssetTools& AssetTools = FModuleManager::LoadModuleChecked<FAssetToolsModule>("AssetTools").Get();
+			AssetTools.FixupReferencers(Redirectors, /*bCheckoutDialogPrompt=*/ true, bDeleteAssets ? ERedirectFixupMode::DeleteFixedUpRedirectors : ERedirectFixupMode::LeaveFixedUpRedirectors);
+		}
+	}
+}
+
+void UAssetDefinition_Redirector::FindTargets(const TArray<UObjectRedirector*>& Redirectors) const
+{
+	TArray<UObject*> ObjectsToSync;
+	for (auto RedirectorIt = Redirectors.CreateConstIterator(); RedirectorIt; ++RedirectorIt)
+	{
+		const UObjectRedirector* Redirector = *RedirectorIt;
+		if ( Redirector && Redirector->DestinationObject )
+		{
+			ObjectsToSync.Add(Redirector->DestinationObject);
+		}
+	}
+
+	if ( ObjectsToSync.Num() > 0 )
+	{
+		IAssetTools& AssetTools = FModuleManager::LoadModuleChecked<FAssetToolsModule>("AssetTools").Get();
+		AssetTools.SyncBrowserToAssets(ObjectsToSync);
+	}
+}
+
+#undef LOCTEXT_NAMESPACE
diff --git a/Engine/Plugins/Editor/EngineAssetDefinitions/Source/Private/AssetDefinition_Redirector.h b/Engine/Plugins/Editor/EngineAssetDefinitions/Source/Private/AssetDefinition_Redirector.h
new file mode 100644
--- /dev/null
+++ b/Engine/Plugins/Editor/EngineAssetDefinitions/Source/Private/AssetDefinition_Redirector.h
@@ -0,0 +1,41 @@
+// Copyright Epic Games, Inc. All Rights Reserved.
+
+#pragma once
+
+#include "CoreMinimal.h"
+#include "AssetDefinitionDefault.h"
+#include "UObject/ObjectRedirector.h"
+
+#include "AssetDefinition_Redirector.generated.h"
+
+struct FToolMenuContext;
+
+UCLASS()
+class UAssetDefinition_Redirector : public UAssetDefinitionDefault
+{
+	GENERATED_BODY()
+
+public:
+	// UAssetDefinition Implementation
+	virtual FText GetAssetDisplayName() const override { return NSLOCTEXT("AssetDefinition", "Redirector", "Redirector"); }
+	virtual FLinearColor GetAssetColor() const override { return FLinearColor(FColor(128, 128, 128)); }
+	virtual TSoftClassPtr<UObject> GetAssetClass() const override { return UObjectRedirector::StaticClass(); }
+	virtual EAssetCommandResult ActivateAssets(const FAssetActivateArgs& ActivateArgs) const override;
+
+	//virtual uint32 GetCategories() override { return EAssetTypeCategories::Misc; }
+
+protected:
+	virtual void OnRegistered() override;
+
+private:
+	void RegisterMenus();
+	
+	/** Handler for when FindTarget is selected */
+	void ExecuteFindTarget(const FToolMenuContext& MenuContext);
+
+	/** Handler for when FixUp is selected */
+	void ExecuteFixUp(const FToolMenuContext& MenuContext, bool bDeleteAssets);
+
+	/** Syncs the content browser to the destination objects for all the supplied redirectors */
+	void FindTargets(const TArray<UObjectRedirector*>& Redirectors) const;
+};
diff --git a/Engine/Plugins/Editor/EngineAssetDefinitions/Source/Private/AssetDefinition_World.cpp b/Engine/Plugins/Editor/EngineAssetDefinitions/Source/Private/AssetDefinition_World.cpp
new file mode 100644
--- /dev/null
+++ b/Engine/Plugins/Editor/EngineAssetDefinitions/Source/Private/AssetDefinition_World.cpp
@@ -0,0 +1,116 @@
+// Copyright Epic Games, Inc. All Rights Reserved.
+
+#include "AssetDefinition_World.h"
+
+#include "FileHelpers.h"
+#include "ThumbnailRendering/WorldThumbnailInfo.h"
+
+#define LOCTEXT_NAMESPACE "UAssetDefinition_World"
+
+TArray<FAssetData> UAssetDefinition_World::PrepareToActivateAssets(const FAssetActivateArgs& ActivateArgs) const
+{
+	TArray<FAssetData> AssetsToOpen;
+	if (ActivateArgs.Assets.Num())
+	{
+		const FAssetData& AssetData = ActivateArgs.Assets[0];
+
+		// If there are any unsaved changes to the current level, see if the user wants to save those first
+		// If they do not wish to save, then we will bail out of opening this asset.
+		constexpr bool bPromptUserToSave = true;
+		constexpr bool bSaveMapPackages = true;
+		constexpr bool bSaveContentPackages = true;
+		if (FEditorFileUtils::SaveDirtyPackages(bPromptUserToSave, bSaveMapPackages, bSaveContentPackages))
+		{
+			// Validate that Asset was saved or isn't loaded meaning it can be loaded
+			const bool bLoad = false;
+			if (UWorld* World = Cast<UWorld>(AssetData.FastGetAsset(bLoad)); !World || !World->GetPackage()->HasAnyPackageFlags(PKG_NewlyCreated))
+			{
+				AssetsToOpen.Add(AssetData);
+			}
+			else
+			{
+				FMessageDialog::Open(EAppMsgType::Ok, LOCTEXT("CannotOpenNewlyCreatedMapWithoutSaving", "The level you are trying to open needs to be saved first."));
+			}
+		}
+	}
+
+	return MoveTemp(AssetsToOpen);
+}
+
+EAssetCommandResult UAssetDefinition_World::OpenAssets(const FAssetOpenArgs& OpenArgs) const
+{
+	for (const FAssetData& WorldAsset : OpenArgs.Assets)
+	{
+		UWorld* World = Cast<UWorld>(WorldAsset.GetAsset());
+		if (World != nullptr && 
+			ensureMsgf(World->GetPackage(), TEXT("World(%s) is not in a package and cannot be opened"), *World->GetFullName()) && 
+			ensureMsgf(!World->GetPackage()->HasAnyPackageFlags(PKG_NewlyCreated), TEXT("World(%s) is unsaved and cannot be opened")))
+		{
+			const FString FileToOpen = FPackageName::LongPackageNameToFilename(World->GetOutermost()->GetName(), FPackageName::GetMapPackageExtension());
+			const bool bLoadAsTemplate = false;
+			const bool bShowProgress = true;
+			FEditorFileUtils::LoadMap(FileToOpen, bLoadAsTemplate, bShowProgress);
+
+			// We can only edit one world at a time... so just break after the first valid world to load
+			return EAssetCommandResult::Handled;
+		}
+	}
+
+	return EAssetCommandResult::Unhandled;
+}
+
+UThumbnailInfo* UAssetDefinition_World::LoadThumbnailInfo(const FAssetData& InAsset) const
+{
+	if (UObject* WorldInstance = InAsset.GetAsset())
+	{
+		UWorld* World = CastChecked<UWorld>(WorldInstance);
+		UThumbnailInfo* ThumbnailInfo = World->ThumbnailInfo;
+		if (ThumbnailInfo == nullptr)
+		{
+			ThumbnailInfo = NewObject<UWorldThumbnailInfo>(World, NAME_None, RF_Transactional);
+			World->ThumbnailInfo = ThumbnailInfo;
+		}
+		
+		return ThumbnailInfo;
+	}
+
+	return nullptr;
+}
+
+FAssetSupportResponse UAssetDefinition_World::CanRename(const FAssetData& InAsset) const
+{
+	if (IsPartitionWorldInUse(InAsset))
+	{
+		return FAssetSupportResponse::Error(LOCTEXT("CanNotRenameWorldInUse", "Cannot rename a partition world while it is used."));
+	}
+	
+	return FAssetSupportResponse::Supported();
+}
+
+FAssetSupportResponse UAssetDefinition_World::CanDuplicate(const FAssetData& InAsset) const
+{
+	if (IsPartitionWorldInUse(InAsset))
+	{
+		return FAssetSupportResponse::Error(LOCTEXT("CanNotDuplicateWorldInUse", "Cannot duplicate a partition world while it is used."));
+	}
+	
+	return FAssetSupportResponse::Supported();
+}
+
+bool UAssetDefinition_World::IsPartitionWorldInUse(const FAssetData& InAsset) const
+{
+	if (ULevel::GetIsLevelPartitionedFromAsset(InAsset))
+	{
+		for (const FWorldContext& WorldContext : GEditor->GetWorldContexts())
+		{
+			if (const UWorld* World = WorldContext.World(); WorldContext.World() && InAsset.PackageName == World->GetPackage()->GetFName())
+			{
+				return true;
+			}
+		}
+	}
+
+	return false;
+}
+
+#undef LOCTEXT_NAMESPACE
diff --git a/Engine/Plugins/Editor/EngineAssetDefinitions/Source/Private/AssetDefinition_World.h b/Engine/Plugins/Editor/EngineAssetDefinitions/Source/Private/AssetDefinition_World.h
new file mode 100644
--- /dev/null
+++ b/Engine/Plugins/Editor/EngineAssetDefinitions/Source/Private/AssetDefinition_World.h
@@ -0,0 +1,34 @@
+// Copyright Epic Games, Inc. All Rights Reserved.
+
+#pragma once
+
+#include "CoreMinimal.h"
+#include "Engine/World.h"
+#include "AssetDefinitionDefault.h"
+
+#include "AssetDefinition_World.generated.h"
+
+UCLASS()
+class UAssetDefinition_World : public UAssetDefinitionDefault
+{
+	GENERATED_BODY()
+
+public:
+	// UAssetDefinition Begin
+	virtual FText GetAssetDisplayName() const override { return NSLOCTEXT("AssetTypeActions", "AssetTypeActions_World", "Level"); }
+	virtual FLinearColor GetAssetColor() const override { return FAppStyle::Get().GetColor("LevelEditor.AssetColor"); }
+	virtual TSoftClassPtr<UObject> GetAssetClass() const override { return UWorld::StaticClass(); }
+	virtual TConstArrayView<FAssetCategoryPath> GetAssetCategories() const override { return { EAssetCategoryPaths::Basic }; }
+
+	virtual TArray<FAssetData> PrepareToActivateAssets(const FAssetActivateArgs& ActivateArgs) const override;
+	virtual EAssetCommandResult OpenAssets(const FAssetOpenArgs& OpenArgs) const override;
+	
+	virtual FAssetSupportResponse CanRename(const FAssetData& InAsset) const override;
+	virtual FAssetSupportResponse CanDuplicate(const FAssetData& InAsset) const override;
+
+	virtual UThumbnailInfo* LoadThumbnailInfo(const FAssetData& InAsset) const override;
+	// UAssetDefinition End
+	
+public:
+	bool IsPartitionWorldInUse(const FAssetData& InAsset) const;
+};
diff --git a/Engine/Plugins/Editor/EngineAssetDefinitions/Source/Private/EngineAssetDefinitionsModule.cpp b/Engine/Plugins/Editor/EngineAssetDefinitions/Source/Private/EngineAssetDefinitionsModule.cpp
new file mode 100644
--- /dev/null
+++ b/Engine/Plugins/Editor/EngineAssetDefinitions/Source/Private/EngineAssetDefinitionsModule.cpp
@@ -0,0 +1,23 @@
+﻿// Copyright Epic Games, Inc. All Rights Reserved.
+
+#include "Modules/ModuleManager.h"
+
+#define LOCTEXT_NAMESPACE "FEngineAssetDefinitionsModule"
+
+class FEngineAssetDefinitionsModule : public IModuleInterface
+{
+public:
+
+	// IModuleInterface interface
+	virtual void StartupModule() override
+	{
+	}
+
+	virtual void ShutdownModule() override
+	{
+	}
+};
+
+#undef LOCTEXT_NAMESPACE
+
+IMPLEMENT_MODULE(FEngineAssetDefinitionsModule, EngineAssetDefinitions);
diff --git a/Engine/Plugins/FX/CascadeToNiagaraConverter/Source/CascadeToNiagaraConverter/Private/CascadeToNiagaraConverterModule.cpp b/Engine/Plugins/FX/CascadeToNiagaraConverter/Source/CascadeToNiagaraConverter/Private/CascadeToNiagaraConverterModule.cpp
--- a/Engine/Plugins/FX/CascadeToNiagaraConverter/Source/CascadeToNiagaraConverter/Private/CascadeToNiagaraConverterModule.cpp
+++ b/Engine/Plugins/FX/CascadeToNiagaraConverter/Source/CascadeToNiagaraConverter/Private/CascadeToNiagaraConverterModule.cpp
@@ -56,7 +56,7 @@ void ICascadeToNiagaraConverterModule::AddMenuExtenderConvertEntry(FMenuBuilder&
 		TArray<UParticleSystem*> CascadeSystems;
 		for (const FAssetData& SelectedAsset : SelectedAssets)
 		{
-			if (SelectedAsset.GetClass()->IsChildOf<UParticleSystem>() == false)
+			if (!SelectedAsset.IsInstanceOf(UParticleSystem::StaticClass()))
 			{
 				return;
 			}
diff --git a/Engine/Source/Developer/AssetTools/AssetTools.Build.cs b/Engine/Source/Developer/AssetTools/AssetTools.Build.cs
--- a/Engine/Source/Developer/AssetTools/AssetTools.Build.cs
+++ b/Engine/Source/Developer/AssetTools/AssetTools.Build.cs
@@ -24,6 +24,7 @@ public class AssetTools : ModuleRules
 
 		PrivateDependencyModuleNames.AddRange(
 			new string[] {
+				"AssetDefinition",
                 "CurveAssetEditor",
 				"Engine",
                 "InputCore",
diff --git a/Engine/Source/Developer/AssetTools/Private/AssetDefinition_AssetTypeActionsProxy.cpp b/Engine/Source/Developer/AssetTools/Private/AssetDefinition_AssetTypeActionsProxy.cpp
new file mode 100644
--- /dev/null
+++ b/Engine/Source/Developer/AssetTools/Private/AssetDefinition_AssetTypeActionsProxy.cpp
@@ -0,0 +1,88 @@
+// Copyright Epic Games, Inc. All Rights Reserved.
+
+#include "AssetDefinition_AssetTypeActionsProxy.h"
+
+#include "Misc/AssetFilterData.h"
+
+void UAssetDefinition_AssetTypeActionsProxy::Initialize(const TSharedRef<IAssetTypeActions>& NewActions)
+{
+	AssetType = NewActions;
+}
+
+FText UAssetDefinition_AssetTypeActionsProxy::GetAssetDisplayName() const
+{
+	return AssetType->GetName();
+}
+
+FText UAssetDefinition_AssetTypeActionsProxy::GetAssetDescription(const FAssetData& AssetData) const
+{
+	return AssetType->GetAssetDescription(AssetData);
+}
+
+TSoftClassPtr<UObject> UAssetDefinition_AssetTypeActionsProxy::GetAssetClass() const
+{
+	if (UClass* SupportedClass = AssetType->GetSupportedClass())
+	{
+		return TSoftClassPtr<UObject>(SupportedClass);
+	}
+
+	return TSoftClassPtr<UObject>(FSoftObjectPath(AssetType->GetClassPathName()));
+}
+
+FLinearColor UAssetDefinition_AssetTypeActionsProxy::GetAssetColor() const
+{
+	return FLinearColor(AssetType->GetTypeColor());
+}
+
+FAssetSupportResponse UAssetDefinition_AssetTypeActionsProxy::CanLocalize(const FAssetData& InAsset) const
+{
+	return AssetType->CanLocalize() ? FAssetSupportResponse::Supported() : FAssetSupportResponse::NotSupported();
+}
+
+bool UAssetDefinition_AssetTypeActionsProxy::CanImport() const
+{
+	return AssetType->IsImportedAsset();
+}
+
+bool UAssetDefinition_AssetTypeActionsProxy::CanMerge() const
+{
+	return AssetType->CanMerge();
+}
+
+EAssetCommandResult UAssetDefinition_AssetTypeActionsProxy::ActivateAssets(const FAssetActivateArgs& ActivateArgs) const
+{
+	EAssetTypeActivationMethod::Type ActivationType = EAssetTypeActivationMethod::Opened;
+	switch(ActivateArgs.ActivationMethod)
+	{
+	case EAssetActivationMethod::DoubleClicked:
+		ActivationType = EAssetTypeActivationMethod::DoubleClicked;
+		break;
+	case EAssetActivationMethod::Opened:
+		ActivationType = EAssetTypeActivationMethod::Opened;
+		break;
+	case EAssetActivationMethod::Previewed:
+		ActivationType = EAssetTypeActivationMethod::Previewed;
+		break;
+	}
+
+	const TArray<UObject*> Assets = ActivateArgs.LoadObjects<UObject>();
+	const bool bResult = AssetType->AssetsActivatedOverride(Assets, ActivationType);
+	return bResult ? EAssetCommandResult::Handled : EAssetCommandResult::Unhandled;
+}
+
+EAssetCommandResult UAssetDefinition_AssetTypeActionsProxy::GetFilters(TArray<FAssetFilterData>& OutFilters) const
+{
+	if (AssetType->CanFilter())
+	{
+		FAssetFilterData Data;
+		Data.Name = AssetType->GetFilterName().ToString();
+		Data.DisplayText = AssetType->GetName();
+		AssetType->BuildBackendFilter(Data.Filter);
+		OutFilters.Add(MoveTemp(Data));
+
+		return EAssetCommandResult::Handled;
+	}
+
+	return EAssetCommandResult::Unhandled;
+}
+
diff --git a/Engine/Source/Developer/AssetTools/Private/AssetDefinition_AssetTypeActionsProxy.h b/Engine/Source/Developer/AssetTools/Private/AssetDefinition_AssetTypeActionsProxy.h
new file mode 100644
--- /dev/null
+++ b/Engine/Source/Developer/AssetTools/Private/AssetDefinition_AssetTypeActionsProxy.h
@@ -0,0 +1,43 @@
+// Copyright Epic Games, Inc. All Rights Reserved.
+
+#pragma once
+
+#include "CoreMinimal.h"
+#include "IAssetTypeActions.h"
+#include "AssetDefinitionDefault.h"
+
+#include "AssetDefinition_AssetTypeActionsProxy.generated.h"
+
+UCLASS(transient)
+class UAssetDefinition_AssetTypeActionsProxy : public UAssetDefinitionDefault
+{
+	GENERATED_BODY() 
+public:
+	UAssetDefinition_AssetTypeActionsProxy() { }
+	
+	void Initialize(const TSharedRef<IAssetTypeActions>& NewActions);
+	
+	const TSharedPtr<IAssetTypeActions>& GetAssetType() const { return AssetType; }
+
+	virtual FText GetAssetDisplayName() const override;
+	virtual TSoftClassPtr<UObject> GetAssetClass() const override;
+	virtual FLinearColor GetAssetColor() const override;
+	virtual FText GetAssetDescription(const FAssetData& AssetData) const override;
+
+	virtual FAssetSupportResponse CanLocalize(const FAssetData& InAsset) const override;
+	virtual bool CanImport() const override;
+	virtual bool CanMerge() const override;
+
+	virtual EAssetCommandResult ActivateAssets(const FAssetActivateArgs& ActivateArgs) const override;
+
+	virtual EAssetCommandResult GetFilters(TArray<FAssetFilterData>& OutFilters) const override;
+
+	// We don't need to implement GetSourceFiles, the default one will / should work for any existing type.
+	// virtual EAssetCommandResult GetSourceFiles(const FAssetSourceFileArgs& SourceFileArgs, TArray<FAssetSourceFile>& OutSourceAssets) const override;
+	
+protected:
+	virtual bool CanRegisterStatically() const override { return false; }
+	
+private:
+	TSharedPtr<IAssetTypeActions> AssetType;
+};
diff --git a/Engine/Source/Developer/AssetTools/Private/AssetTools.cpp b/Engine/Source/Developer/AssetTools/Private/AssetTools.cpp
--- a/Engine/Source/Developer/AssetTools/Private/AssetTools.cpp
+++ b/Engine/Source/Developer/AssetTools/Private/AssetTools.cpp
@@ -31,6 +31,7 @@
 #include "AssetToolsModule.h"
 #include "Framework/MultiBox/MultiBoxBuilder.h"
 #include "ToolMenus.h"
+#include "AssetDefinition_AssetTypeActionsProxy.h"
 #include "IClassTypeActions.h"
 #include "ClassViewerModule.h"
 #include "ClassViewerFilter.h"
@@ -101,7 +102,6 @@
 #include "AssetTypeActions/AssetTypeActions_PoseAsset.h"
 #include "AssetTypeActions/AssetTypeActions_PreviewMeshCollection.h"
 #include "AssetTypeActions/AssetTypeActions_ProceduralFoliageSpawner.h"
-#include "AssetTypeActions/AssetTypeActions_Redirector.h"
 #include "AssetTypeActions/AssetTypeActions_Rig.h"
 #include "AssetTypeActions/AssetTypeActions_Skeleton.h"
 #include "AssetTypeActions/AssetTypeActions_SlateBrush.h"
@@ -120,7 +120,6 @@
 #include "AssetTypeActions/AssetTypeActions_TouchInterface.h"
 #include "AssetTypeActions/AssetTypeActions_VectorFieldAnimated.h"
 #include "AssetTypeActions/AssetTypeActions_VectorFieldStatic.h"
-#include "AssetTypeActions/AssetTypeActions_World.h"
 #include "WorldPartition/WorldPartition.h"
 #include "SDiscoveringAssetsDialog.h"
 #include "AssetFixUpRedirectors.h"
@@ -179,10 +178,291 @@
 #include "ObjectTools.h"
 #endif
 
+#include "AssetDefinition.h"
+#include "AssetDefinitionRegistry.h"
+#include "Algo/AnyOf.h"
+#include "Misc/AssetFilterData.h"
+
 #include UE_INLINE_GENERATED_CPP_BY_NAME(AssetTools)
 
 #define LOCTEXT_NAMESPACE "AssetTools"
 
+class FAssetDefinitionProxy : public FAssetTypeActions_Base
+{
+public:
+	FAssetDefinitionProxy(UAssetDefinition* AssetDefinition)
+		: AssetDefinitionPtr(AssetDefinition)
+	{
+	}
+	
+	virtual bool IsAssetDefinitionInDisguise() const override
+	{
+		return true;
+	}
+
+	virtual bool ShouldCallGetActions() const override { return false; }
+
+	UAssetDefinition* GetAssetDefinition() const { return AssetDefinitionPtr.Get(); }
+	
+	virtual FText GetName() const override { return AssetDefinitionPtr.Get()->GetAssetDisplayName(); }
+	virtual UClass* GetSupportedClass() const override { return AssetDefinitionPtr.Get()->GetAssetClass().LoadSynchronous(); }
+	virtual FColor GetTypeColor() const override { return AssetDefinitionPtr.Get()->GetAssetColor().ToFColor(true); }
+	
+	virtual bool CanLocalize() const override { return AssetDefinitionPtr.Get()->CanLocalize(FAssetData()).IsSupported(); }
+	virtual bool IsImportedAsset() const override { return AssetDefinitionPtr.Get()->CanImport(); }
+	virtual bool CanMerge() const override { return AssetDefinitionPtr.Get()->CanMerge(); }
+	
+	virtual bool CanRename(const FAssetData& InAsset, FText* OutErrorMsg) const override
+    {
+		FAssetSupportResponse Response = AssetDefinitionPtr.Get()->CanRename(InAsset);
+		if (OutErrorMsg)
+		{
+			(*OutErrorMsg) = Response.GetErrorText();
+		}
+		return Response.IsSupported();
+    }
+
+    virtual bool CanDuplicate(const FAssetData& InAsset, FText* OutErrorMsg) const override
+    {
+    	FAssetSupportResponse Response = AssetDefinitionPtr.Get()->CanDuplicate(InAsset);
+		if (OutErrorMsg)
+		{
+			(*OutErrorMsg) = Response.GetErrorText();
+		}
+        return Response.IsSupported();
+    }
+	
+	virtual uint32 GetCategories() override
+	{
+		static IAssetTools& AssetTools = FModuleManager::LoadModuleChecked<FAssetToolsModule>("AssetTools").Get();
+
+		uint32 CategoryBits = 0;
+		for (const FAssetCategoryPath& Category : AssetDefinitionPtr.Get()->GetAssetCategories())
+		{
+			CategoryBits |= AssetTools.FindAdvancedAssetCategory(Category.GetCategory());
+		}
+
+		return EAssetTypeCategories::Misc;
+	}
+
+private:
+	TArray<FText> SubMenus;
+	bool SubmenusInitialized = false;
+
+public:
+	/** Returns array of sub-menu names that this asset type is parented under in the Asset Creation Context Menu. */
+	virtual const TArray<FText>& GetSubMenus() const override
+	{
+		static IAssetTools& AssetTools = FModuleManager::LoadModuleChecked<FAssetToolsModule>("AssetTools").Get();
+
+		if (!SubmenusInitialized)
+		{
+			for (const FAssetCategoryPath& Category : AssetDefinitionPtr.Get()->GetAssetCategories())
+			{
+				if (Category.HasSubCategory())
+				{
+					const_cast<FAssetDefinitionProxy*>(this)->SubMenus.Add(Category.GetSubCategoryText());
+				}
+			}
+		}
+    
+		return SubMenus;
+	}
+
+	virtual void OpenAssetEditor( const TArray<UObject*>& InObjects, TSharedPtr<class IToolkitHost> EditWithinLevelEditor) override
+	{
+		OpenAssetEditor(InObjects, EAssetTypeActivationOpenedMethod::Edit, EditWithinLevelEditor);
+	}
+	
+	virtual void OpenAssetEditor(const TArray<UObject*>& InObjects, const EAssetTypeActivationOpenedMethod OpenedMethod, TSharedPtr<IToolkitHost> EditWithinLevelEditor = TSharedPtr<IToolkitHost>()) override
+	{
+		TArray<FAssetData> Assets;
+		Algo::Transform(InObjects, Assets, [](UObject* Object){ return FAssetData(Object); });
+		
+		FAssetOpenArgs OpenArgs;
+		OpenArgs.OpenMethod = OpenedMethod == EAssetTypeActivationOpenedMethod::Edit ? EAssetOpenMethod::Edit : EAssetOpenMethod::View;
+		OpenArgs.ToolkitHost = EditWithinLevelEditor;
+		OpenArgs.Assets = Assets;
+		AssetDefinitionPtr.Get()->OpenAssets(OpenArgs);		
+	}
+
+	virtual TArray<FAssetData> GetValidAssetsForPreviewOrEdit(TArrayView<const FAssetData> InAssetDatas, bool bIsPreview) override
+	{
+		FAssetActivateArgs ActivateArgs;
+		ActivateArgs.ActivationMethod = bIsPreview ? EAssetActivationMethod::Opened : EAssetActivationMethod::Previewed;
+		ActivateArgs.Assets = InAssetDatas;
+
+		return AssetDefinitionPtr.Get()->PrepareToActivateAssets(ActivateArgs);
+	}
+	
+	virtual bool AssetsActivatedOverride(const TArray<UObject*>& InObjects, EAssetTypeActivationMethod::Type ActivationType) override
+	{
+		EAssetActivationMethod ActivationMethod = EAssetActivationMethod::Opened;
+		switch(ActivationType)
+		{
+		case EAssetTypeActivationMethod::DoubleClicked:
+			ActivationMethod = EAssetActivationMethod::DoubleClicked;
+			break;
+		case EAssetTypeActivationMethod::Opened:
+			ActivationMethod = EAssetActivationMethod::Opened;
+			break;
+		case EAssetTypeActivationMethod::Previewed:
+			ActivationMethod = EAssetActivationMethod::Previewed;
+			break;
+		}
+		
+		TArray<FAssetData> Assets;
+		Algo::Transform(InObjects, Assets, [](UObject* Object){ return FAssetData(Object);});
+		
+		FAssetActivateArgs Args;
+		Args.ActivationMethod = ActivationMethod;
+		Args.Assets = Assets;
+		return AssetDefinitionPtr->ActivateAssets(Args) == EAssetCommandResult::Handled;
+	}
+
+	virtual bool CanFilter() override
+	{
+		TArray<FAssetFilterData> Filters;
+		AssetDefinitionPtr->GetFilters(Filters);
+		return Filters.Num() > 0;
+	}
+
+	virtual FName GetFilterName() const override
+	{
+		TArray<FAssetFilterData> Filters;
+		AssetDefinitionPtr->GetFilters(Filters);
+		return Filters.Num() > 0 ? FName(*Filters[0].Name) : NAME_None;
+	}
+
+	virtual void BuildBackendFilter(FARFilter& InFilter) override
+	{
+		TArray<FAssetFilterData> Filters;
+		AssetDefinitionPtr->GetFilters(Filters);
+		if (Filters.Num() > 0)
+		{
+			InFilter = Filters[0].Filter;
+		}
+	}
+
+	virtual FText GetDisplayNameFromAssetData(const FAssetData& AssetData) const override
+	{
+		return AssetDefinitionPtr->GetAssetDisplayName(AssetData);
+	}
+
+	virtual UThumbnailInfo* GetThumbnailInfo(UObject* Object) const override
+	{
+		return AssetDefinitionPtr->LoadThumbnailInfo(FAssetData(Object));
+	}
+
+	virtual const FSlateBrush* GetThumbnailBrush(const FAssetData& InAssetData, const FName InClassName) const override
+	{
+		return AssetDefinitionPtr->GetThumbnailBrush(InAssetData, InClassName);
+	}
+	
+	virtual const FSlateBrush* GetIconBrush(const FAssetData& InAssetData, const FName InClassName) const override
+	{
+		return AssetDefinitionPtr->GetThumbnailBrush(InAssetData, InClassName);
+	}
+
+	virtual FTopLevelAssetPath GetClassPathName() const override
+	{
+		return AssetDefinitionPtr.Get()->GetAssetClass().ToSoftObjectPath().GetAssetPath();
+	}
+
+	/** Begins a merge operation for InObject (automatically determines remote/base versions needed to resolve) */
+	virtual void Merge( UObject* InObject ) override
+	{
+		FAssetMergeArgs MergeArgs;
+		MergeArgs.LocalAsset = InObject;
+
+		AssetDefinitionPtr.Get()->Merge(MergeArgs);
+	}
+
+	/** Begins a merge between the specified assets */
+	virtual void Merge(UObject* BaseAsset, UObject* RemoteAsset, UObject* LocalAsset, const FOnMergeResolved& ResolutionCallback) override
+	{
+		FAssetMergeArgs MergeArgs;
+		MergeArgs.LocalAsset = LocalAsset;
+		MergeArgs.BaseAsset = BaseAsset;
+		MergeArgs.RemoteAsset = RemoteAsset;
+		MergeArgs.ResolutionCallback = FOnAssetMergeResolved::CreateLambda([ResolutionCallback](const FAssetMergeResults& Results)
+		{
+			ResolutionCallback.Execute(Results.MergedPackage, static_cast<EMergeResult::Type>(Results.Result));
+		});
+
+		AssetDefinitionPtr.Get()->Merge(MergeArgs);
+	}
+	
+	virtual void PerformAssetDiff(UObject* OldAsset, UObject* NewAsset, const struct FRevisionInfo& OldRevision, const struct FRevisionInfo& NewRevision) const override
+	{
+		FAssetDiffArgs DiffArgs;
+		DiffArgs.OldAsset = OldAsset;
+		DiffArgs.NewAsset = NewAsset;
+		DiffArgs.OldRevision = OldRevision;
+		DiffArgs.NewRevision = NewRevision;
+		
+		AssetDefinitionPtr.Get()->PerformAssetDiff(DiffArgs);
+	}
+
+	/** Returns additional tooltip information for the specified asset, if it has any (otherwise return the null widget) */
+	virtual FText GetAssetDescription(const FAssetData& AssetData) const override
+	{
+		ensureMsgf(false, TEXT("This code path is not expected to be called any more.  Expected UAssetDefinitionRegistry::Get()->GetAssetDefinitionForAsset(YourAsset)->GetAssetDescription to be called directly."));
+        return FText::GetEmpty();
+	}
+
+	/** Collects the resolved source paths for the imported assets */
+	virtual void GetResolvedSourceFilePaths(const TArray<UObject*>& TypeAssets, TArray<FString>& OutSourceFilePaths) const override
+	{
+		ensureMsgf(false, TEXT("This code path is not expected to be called any more.  Expected UAssetDefinitionRegistry::Get()->GetAssetDefinitionForAsset(YourAsset)->GetSourceFiles to be called directly."));
+	}
+	
+	/** Collects the source file labels for the imported assets */
+	virtual void GetSourceFileLabels(const TArray<UObject*>& TypeAssets, TArray<FString>& OutSourceFileLabels) const override
+	{
+		ensureMsgf(false, TEXT("This code path is not expected to be called any more.  Expected UAssetDefinitionRegistry::Get()->GetAssetDefinitionForAsset(YourAsset)->GetSourceFiles to be called directly."));
+	}
+
+	/** Does this asset support edit or view methods? */
+	virtual bool SupportsOpenedMethod(const EAssetTypeActivationOpenedMethod OpenedMethod) const override
+	{
+		FAssetOpenSupportArgs SupportArgs;
+		SupportArgs.OpenMethod = OpenedMethod == EAssetTypeActivationOpenedMethod::Edit ? EAssetOpenMethod::Edit : EAssetOpenMethod::View;
+		return AssetDefinitionPtr.Get()->GetAssetOpenSupport(SupportArgs).IsSupported;
+	}
+	
+	/** Optionally returns a custom widget to overlay on top of this assets' thumbnail */
+    virtual TSharedPtr<class SWidget> GetThumbnailOverlay(const FAssetData& AssetData) const override
+    {
+    	return AssetDefinitionPtr.Get()->GetThumbnailOverlay(AssetData);
+    }
+	
+	virtual EThumbnailPrimType GetDefaultThumbnailPrimitiveType(UObject* Asset) const override
+	{
+		ensureMsgf(false, TEXT("This code path is not expected to be called any more.  Expected UAssetDefinitionRegistry::Get()->LoadThumbnail(YourAsset).  The default thumbnail primitive is now a property on USceneThumbnailInfoWithPrimitive.  The default only applies to this thumbnail type, so it has been relocated as an implementation detail of it."));
+		return EThumbnailPrimType::TPT_Sphere;
+	}
+
+	/** @return True if we should force world-centric mode for newly-opened assets */
+	virtual bool ShouldForceWorldCentric() override
+	{
+		FAssetOpenSupportArgs SupportArgs;
+        SupportArgs.OpenMethod = EAssetOpenMethod::Edit;
+        if (AssetDefinitionPtr.Get()->GetAssetOpenSupport(SupportArgs).RequiredToolkitMode.IsSet())
+        {
+        	return AssetDefinitionPtr.Get()->GetAssetOpenSupport(SupportArgs).RequiredToolkitMode.GetValue() == EToolkitMode::WorldCentric; 
+        }
+
+		return false;
+	}
+	
+private:
+
+	TWeakObjectPtr<UAssetDefinition> AssetDefinitionPtr;
+};
+
+//----------------------------------------------------------------------------------------------------------------------
+
 namespace UE::AssetTools::Private
 {
 	bool bUseNewPackageMigration = true;
@@ -920,7 +1200,6 @@ UAssetToolsImpl::UAssetToolsImpl(const FObjectInitializer& ObjectInitializer)
 			}
 			AssetClassPermissionList[i]->AddAllowListItem("AssetToolsConfigFile", Type);
 		}
-		AssetClassPermissionList[i]->OnFilterChanged().AddUObject(this, &UAssetToolsImpl::AssetClassPermissionListChanged, (EAssetClassAction)i);
 	}
 
 	TArray<FString> DenyListedViewPath;
@@ -940,16 +1219,16 @@ UAssetToolsImpl::UAssetToolsImpl(const FObjectInitializer& ObjectInitializer)
 	FPackageName::OnContentPathMounted().AddUObject(this, &UAssetToolsImpl::OnContentPathMounted);
 
 	// Register the built-in advanced categories
-	AllocatedCategoryBits.Add(TEXT("_BuiltIn_0"), FAdvancedAssetCategory(EAssetTypeCategories::Animation, LOCTEXT("AnimationAssetCategory", "Animation")));
-	AllocatedCategoryBits.Add(TEXT("_BuiltIn_1"), FAdvancedAssetCategory(EAssetTypeCategories::Blueprint, LOCTEXT("BlueprintAssetCategory", "Blueprints")));
-	AllocatedCategoryBits.Add(TEXT("_BuiltIn_2"), FAdvancedAssetCategory(EAssetTypeCategories::Materials, LOCTEXT("MaterialAssetCategory", "Materials")));
-	AllocatedCategoryBits.Add(TEXT("_BuiltIn_3"), FAdvancedAssetCategory(EAssetTypeCategories::Sounds, LOCTEXT("SoundAssetCategory", "Sounds")));
-	AllocatedCategoryBits.Add(TEXT("_BuiltIn_4"), FAdvancedAssetCategory(EAssetTypeCategories::Physics, LOCTEXT("PhysicsAssetCategory", "Physics")));
-	AllocatedCategoryBits.Add(TEXT("_BuiltIn_5"), FAdvancedAssetCategory(EAssetTypeCategories::UI, LOCTEXT("UserInterfaceAssetCategory", "User Interface")));
-	AllocatedCategoryBits.Add(TEXT("_BuiltIn_6"), FAdvancedAssetCategory(EAssetTypeCategories::Misc, LOCTEXT("MiscellaneousAssetCategory", "Miscellaneous")));
-	AllocatedCategoryBits.Add(TEXT("_BuiltIn_7"), FAdvancedAssetCategory(EAssetTypeCategories::Gameplay, LOCTEXT("GameplayAssetCategory", "Gameplay")));
-	AllocatedCategoryBits.Add(TEXT("_BuiltIn_8"), FAdvancedAssetCategory(EAssetTypeCategories::Media, LOCTEXT("MediaAssetCategory", "Media")));
-	AllocatedCategoryBits.Add(TEXT("_BuiltIn_9"), FAdvancedAssetCategory(EAssetTypeCategories::Textures, LOCTEXT("TextureAssetCategory", "Textures")));
+	AllocatedCategoryBits.Add(TEXT("Animation"), FAdvancedAssetCategory(EAssetTypeCategories::Animation, LOCTEXT("AnimationAssetCategory", "Animation")));
+	AllocatedCategoryBits.Add(TEXT("Blueprint"), FAdvancedAssetCategory(EAssetTypeCategories::Blueprint, LOCTEXT("BlueprintAssetCategory", "Blueprints")));
+	AllocatedCategoryBits.Add(TEXT("Material"), FAdvancedAssetCategory(EAssetTypeCategories::Materials, LOCTEXT("MaterialAssetCategory", "Materials")));
+	AllocatedCategoryBits.Add(TEXT("Audio"), FAdvancedAssetCategory(EAssetTypeCategories::Sounds, LOCTEXT("SoundAssetCategory", "Sounds")));
+	AllocatedCategoryBits.Add(TEXT("Physics"), FAdvancedAssetCategory(EAssetTypeCategories::Physics, LOCTEXT("PhysicsAssetCategory", "Physics")));
+	AllocatedCategoryBits.Add(TEXT("User Interface"), FAdvancedAssetCategory(EAssetTypeCategories::UI, LOCTEXT("UserInterfaceAssetCategory", "User Interface")));
+	AllocatedCategoryBits.Add(TEXT("Misc"), FAdvancedAssetCategory(EAssetTypeCategories::Misc, LOCTEXT("MiscellaneousAssetCategory", "Miscellaneous")));
+	AllocatedCategoryBits.Add(TEXT("Gameplay"), FAdvancedAssetCategory(EAssetTypeCategories::Gameplay, LOCTEXT("GameplayAssetCategory", "Gameplay")));
+	AllocatedCategoryBits.Add(TEXT("Media"), FAdvancedAssetCategory(EAssetTypeCategories::Media, LOCTEXT("MediaAssetCategory", "Media")));
+	AllocatedCategoryBits.Add(TEXT("Texture"), FAdvancedAssetCategory(EAssetTypeCategories::Textures, LOCTEXT("TextureAssetCategory", "Textures")));
 
 	EAssetTypeCategories::Type FoliageCategoryBit = RegisterAdvancedAssetCategory(FName(TEXT("Foliage")), LOCTEXT("FoliageAssetCategory", "Foliage"));
 	EAssetTypeCategories::Type InputCategoryBit = RegisterAdvancedAssetCategory(FName(TEXT("Input")), LOCTEXT("InputAssetsCategory", "Input"));
@@ -982,7 +1261,7 @@ UAssetToolsImpl::UAssetToolsImpl(const FObjectInitializer& ObjectInitializer)
 	RegisterAssetTypeActions(MakeShareable(new FAssetTypeActions_CurveVector));
 	RegisterAssetTypeActions(MakeShareable(new FAssetTypeActions_CurveLinearColor));
 	RegisterAssetTypeActions(MakeShareable(new FAssetTypeActions_CurveLinearColorAtlas));
-	RegisterAssetTypeActions(MakeShareable(new FAssetTypeActions_DataAsset));
+	//RegisterAssetTypeActions(MakeShareable(new FAssetTypeActions_DataAsset));
 	RegisterAssetTypeActions(MakeShareable(new FAssetTypeActions_DataLayer));
 	RegisterAssetTypeActions(MakeShareable(new FAssetTypeActions_DataTable));
 	RegisterAssetTypeActions(MakeShareable(new FAssetTypeActions_CompositeDataTable));
@@ -1021,7 +1300,6 @@ UAssetToolsImpl::UAssetToolsImpl(const FObjectInitializer& ObjectInitializer)
 	RegisterAssetTypeActions(MakeShareable(new FAssetTypeActions_PhysicsAsset));
 	RegisterAssetTypeActions(MakeShareable(new FAssetTypeActions_PreviewMeshCollection));
 	RegisterAssetTypeActions(MakeShareable(new FAssetTypeActions_ProceduralFoliageSpawner(FoliageCategoryBit)));
-	RegisterAssetTypeActions(MakeShareable(new FAssetTypeActions_Redirector));
 	RegisterAssetTypeActions(MakeShareable(new FAssetTypeActions_Rig));
 	RegisterAssetTypeActions(MakeShareable(new FAssetTypeActions_SkeletalMesh));
 	RegisterAssetTypeActions(MakeShareable(new FAssetTypeActions_Skeleton));
@@ -1044,7 +1322,6 @@ UAssetToolsImpl::UAssetToolsImpl(const FObjectInitializer& ObjectInitializer)
 	RegisterAssetTypeActions(MakeShareable(new FAssetTypeActions_VectorField));
 	RegisterAssetTypeActions(MakeShareable(new FAssetTypeActions_VectorFieldAnimated));
 	RegisterAssetTypeActions(MakeShareable(new FAssetTypeActions_VectorFieldStatic));
-	RegisterAssetTypeActions(MakeShareable(new FAssetTypeActions_World));
 	RegisterAssetTypeActions(MakeShareable(new FAssetTypeActions_HapticFeedbackEffectBuffer));
 	RegisterAssetTypeActions(MakeShareable(new FAssetTypeActions_HapticFeedbackEffectCurve));
 	RegisterAssetTypeActions(MakeShareable(new FAssetTypeActions_HapticFeedbackEffectSoundWave));
@@ -1055,28 +1332,65 @@ UAssetToolsImpl::UAssetToolsImpl(const FObjectInitializer& ObjectInitializer)
 
 void UAssetToolsImpl::RegisterAssetTypeActions(const TSharedRef<IAssetTypeActions>& NewActions)
 {
-	bool bSupported = false;
-	if (const UClass* SupportedClass = NewActions->GetSupportedClass())
-	{
-		bSupported = GetAssetClassPathPermissionList(EAssetClassAction::CreateAsset)->PassesFilter(SupportedClass->GetClassPathName().ToString());
-	}
-	else
+	AssetTypeActionsList.Add(NewActions);
+
+	if (!NewActions->IsAssetDefinitionInDisguise())
 	{
-		bSupported = !NewActions->GetFilterName().IsNone();
+		UAssetDefinition_AssetTypeActionsProxy* Proxy = NewObject<UAssetDefinition_AssetTypeActionsProxy>();
+		Proxy->Initialize(NewActions);
+		
+		UAssetDefinitionRegistry::Get()->RegisterAssetDefinition(Proxy);
 	}
+}
 
-	NewActions->SetSupported(bSupported);
+void UAssetToolsImpl::SyncAssetTypesToAssetDefinitions() const
+{
+	static int32 CachedAmount = 0;
 
-	AssetTypeActionsList.Add(NewActions);
+	TArray<TObjectPtr<UAssetDefinition>> AssetDefinitions = UAssetDefinitionRegistry::Get()->GetAllAssetDefinitions();
+	if (CachedAmount != AssetDefinitions.Num())
+	{
+		CachedAmount = AssetDefinitions.Num();		
+		for (UAssetDefinition* AssetDefinition : AssetDefinitions)
+		{
+			if (!AssetDefinition->IsA<UAssetDefinition_AssetTypeActionsProxy>())
+			{
+				const TSoftClassPtr<UObject> AssetClass = AssetDefinition->GetAssetClass();
+				const bool ActionDefinitionAlreadyProxied = 
+					Algo::AnyOf(AssetTypeActionsList, [AssetClass](const TSharedRef<IAssetTypeActions>& Actions){ return Actions->GetSupportedClass() == AssetClass; });
+			
+				if (!ActionDefinitionAlreadyProxied)
+				{
+					TSharedRef<FAssetDefinitionProxy> Proxy = MakeShared<FAssetDefinitionProxy>(AssetDefinition);
+					// Cache the asset definition proxy.
+					const_cast<UAssetToolsImpl*>(this)->RegisterAssetTypeActions(Proxy);
+				}
+			}
+		}
+	}
 }
 
 void UAssetToolsImpl::UnregisterAssetTypeActions(const TSharedRef<IAssetTypeActions>& ActionsToRemove)
 {
 	AssetTypeActionsList.Remove(ActionsToRemove);
+
+	// Now also remove the proxy AssetDefinition we registered for this type.
+	for (UAssetDefinition* AssetDefinition : UAssetDefinitionRegistry::Get()->GetAllAssetDefinitions())
+	{
+		if (UAssetDefinition_AssetTypeActionsProxy* Proxy = Cast<UAssetDefinition_AssetTypeActionsProxy>(AssetDefinition))
+		{
+			if (Proxy->GetAssetType() == ActionsToRemove)
+			{
+				UAssetDefinitionRegistry::Get()->UnregisterAssetDefinition(AssetDefinition);
+				break;
+			}
+		}
+	}
 }
 
 void UAssetToolsImpl::GetAssetTypeActionsList( TArray<TWeakPtr<IAssetTypeActions>>& OutAssetTypeActionsList ) const
 {
+	SyncAssetTypesToAssetDefinitions();
 	for (auto ActionsIt = AssetTypeActionsList.CreateConstIterator(); ActionsIt; ++ActionsIt)
 	{
 		OutAssetTypeActionsList.Add(*ActionsIt);
@@ -1092,6 +1406,7 @@ TWeakPtr<IAssetTypeActions> UAssetToolsImpl::GetAssetTypeActionsForClass(const U
 
 	TSharedPtr<IAssetTypeActions> MostDerivedAssetTypeActions;
 
+	SyncAssetTypesToAssetDefinitions();
 	for (int32 TypeActionsIdx = 0; TypeActionsIdx < AssetTypeActionsList.Num(); ++TypeActionsIdx)
 	{
 		TSharedRef<IAssetTypeActions> TypeActions = AssetTypeActionsList[TypeActionsIdx];
@@ -1109,10 +1424,45 @@ TWeakPtr<IAssetTypeActions> UAssetToolsImpl::GetAssetTypeActionsForClass(const U
 	return MostDerivedAssetTypeActions;
 }
 
+bool UAssetToolsImpl::CanLocalize(const UClass* Class) const
+{
+	if (const UAssetDefinition* AssetDefinition = UAssetDefinitionRegistry::Get()->GetAssetDefinitionForClass(Class))
+	{
+		return AssetDefinition->CanLocalize(FAssetData()).IsSupported();
+	}
+	else
+	{
+		if (TSharedPtr<IAssetTypeActions> AssetActions = GetAssetTypeActionsForClass(Class).Pin())
+		{
+			return AssetActions->CanLocalize();
+		}
+	}
+
+	return false;
+}
+
+TOptional<FLinearColor> UAssetToolsImpl::GetTypeColor(const UClass* Class) const
+{
+	if (const UAssetDefinition* AssetDefinition = UAssetDefinitionRegistry::Get()->GetAssetDefinitionForClass(Class))
+	{
+		return AssetDefinition->GetAssetColor();
+	}
+	else
+	{
+		if (TSharedPtr<IAssetTypeActions> AssetActions = GetAssetTypeActionsForClass(Class).Pin())
+		{
+			return FLinearColor(AssetActions->GetTypeColor());
+		}
+	}
+
+	return TOptional<FLinearColor>();
+}
+
 TArray<TWeakPtr<IAssetTypeActions>> UAssetToolsImpl::GetAssetTypeActionsListForClass(const UClass* Class) const
 {
 	TArray<TWeakPtr<IAssetTypeActions>> ResultAssetTypeActionsList;
 
+	SyncAssetTypesToAssetDefinitions();
 	for (int32 TypeActionsIdx = 0; TypeActionsIdx < AssetTypeActionsList.Num(); ++TypeActionsIdx)
 	{
 		TSharedRef<IAssetTypeActions> TypeActions = AssetTypeActionsList[TypeActionsIdx];
@@ -4891,17 +5241,22 @@ void UAssetToolsImpl::AdvancedCopyPackages_ReportConfirmed(FAdvancedCopyParams C
 bool UAssetToolsImpl::IsAssetClassSupported(const UClass* AssetClass) const
 {
 	TWeakPtr<IAssetTypeActions> AssetTypeActions = GetAssetTypeActionsForClass(AssetClass);
-	if (!AssetTypeActions.IsValid())
+	if (AssetTypeActions.IsValid())
 	{
-		return false;
-	}
+		bool bSupported = false;
+		if (const UClass* SupportedClass = AssetTypeActions.Pin()->GetSupportedClass())
+		{
+			bSupported = GetAssetClassPathPermissionList(EAssetClassAction::CreateAsset)->PassesFilter(SupportedClass->GetClassPathName().ToString());
+		}
+		//else
+		//{
+		//	bSupported = !AssetTypeActions.Pin()->GetFilterName().IsNone();
+		//}
 
-	if (AssetTypeActions.Pin()->IsSupported() == false)
-	{
-		return false;
+		return bSupported;
 	}
-
-	return true;
+	
+	return false;
 }
 
 TArray<UFactory*> UAssetToolsImpl::GetNewAssetFactories() const
@@ -5008,7 +5363,7 @@ TSharedRef<FNamePermissionList>& UAssetToolsImpl::GetAssetClassPermissionList(EA
 	return Empty;
 }
 
-TSharedRef<FPathPermissionList>& UAssetToolsImpl::GetAssetClassPathPermissionList(EAssetClassAction AssetClassAction)
+const TSharedRef<FPathPermissionList>& UAssetToolsImpl::GetAssetClassPathPermissionList(EAssetClassAction AssetClassAction) const
 {
 	if (AssetClassAction < EAssetClassAction::AllAssetActions)
 	{
@@ -5024,31 +5379,6 @@ TSet<EBlueprintType>& UAssetToolsImpl::GetAllowedBlueprintTypes()
 	return AllowedBlueprintTypes;
 }
 
-void UAssetToolsImpl::AssetClassPermissionListChanged(EAssetClassAction AssetClassAction)
-{
-	switch (AssetClassAction)
-	{
-	case EAssetClassAction::CreateAsset:
-		for (TSharedRef<IAssetTypeActions>& ActionsIt : AssetTypeActionsList)
-		{
-			bool bSupported = false;
-			if (const UClass* SupportedClass = ActionsIt->GetSupportedClass())
-			{
-				bSupported = GetAssetClassPathPermissionList(AssetClassAction)->PassesFilter(SupportedClass->GetPathName());
-			}
-			else
-			{
-				bSupported = !ActionsIt->GetFilterName().IsNone();
-			}
-
-			ActionsIt->SetSupported(bSupported);
-		}
-		break;
-	default:
-		break;
-	}
-}
-
 void UAssetToolsImpl::AddSubContentDenyList(const FString& InMount)
 {
 	for (const FString& SubContentPath : SubContentDenyListPaths)
diff --git a/Engine/Source/Developer/AssetTools/Private/AssetTools.h b/Engine/Source/Developer/AssetTools/Private/AssetTools.h
--- a/Engine/Source/Developer/AssetTools/Private/AssetTools.h
+++ b/Engine/Source/Developer/AssetTools/Private/AssetTools.h
@@ -69,6 +69,9 @@ public:
 	virtual void UnregisterAssetTypeActions(const TSharedRef<IAssetTypeActions>& ActionsToRemove) override;
 	virtual void GetAssetTypeActionsList( TArray<TWeakPtr<IAssetTypeActions>>& OutAssetTypeActionsList ) const override;
 	virtual TWeakPtr<IAssetTypeActions> GetAssetTypeActionsForClass(const UClass* Class) const override;
+	virtual bool CanLocalize(const UClass* Class) const;
+	virtual TOptional<FLinearColor> GetTypeColor(const UClass* Class) const override;
+	
 	virtual TArray<TWeakPtr<IAssetTypeActions>> GetAssetTypeActionsListForClass(const UClass* Class) const override;
 	virtual EAssetTypeCategories::Type RegisterAdvancedAssetCategory(FName CategoryKey, FText CategoryDisplayName) override;
 	virtual EAssetTypeCategories::Type FindAdvancedAssetCategory(FName CategoryKey) const override;
@@ -129,7 +132,7 @@ public:
 	virtual TSharedRef<FNamePermissionList>& GetAssetClassPermissionList() override;
 	UE_DEPRECATED(5.1, "Class names are now represented by path names. Please use GetAssetClassPathPermissionList.")
 	TSharedRef<FNamePermissionList>& GetAssetClassPermissionList(EAssetClassAction AssetClassAction);
-	virtual TSharedRef<FPathPermissionList>& GetAssetClassPathPermissionList(EAssetClassAction AssetClassAction) override;
+	virtual const TSharedRef<FPathPermissionList>& GetAssetClassPathPermissionList(EAssetClassAction AssetClassAction) const override;
 	virtual TSet<EBlueprintType>& GetAllowedBlueprintTypes() override;
 	virtual TSharedRef<FPathPermissionList>& GetFolderPermissionList() override;
 	virtual TSharedRef<FPathPermissionList>& GetWritableFolderPermissionList() override;
@@ -193,9 +196,6 @@ private:
 
 	UObject* PerformDuplicateAsset(const FString& AssetName, const FString& PackagePath, UObject* OriginalObject, bool bWithDialog);
 
-	/** Internal method that performs actions when asset class deny list filter changes */
-	void AssetClassPermissionListChanged(EAssetClassAction AssetClassAction);
-
 	/**
 	 * Add sub content deny list filter for a new mount point
 	 * @param InMount The mount point
@@ -208,6 +208,9 @@ private:
 	/** Implementation for the import with dialog functions */
 	TArray<UObject*> ImportAssetsWithDialogImplementation(const FString& DestinationPath, bool bAllowAsyncImport);
 
+	/** Make sure we're not syncing */
+	void SyncAssetTypesToAssetDefinitions() const;
+
 private:
 	/** The list of all registered AssetTypeActions */
 	TArray<TSharedRef<IAssetTypeActions>> AssetTypeActionsList;
diff --git a/Engine/Source/Developer/AssetTools/Private/AssetTypeActions/AssetTypeActions_MaterialInterface.cpp b/Engine/Source/Developer/AssetTools/Private/AssetTypeActions/AssetTypeActions_MaterialInterface.cpp
--- a/Engine/Source/Developer/AssetTools/Private/AssetTypeActions/AssetTypeActions_MaterialInterface.cpp
+++ b/Engine/Source/Developer/AssetTools/Private/AssetTypeActions/AssetTypeActions_MaterialInterface.cpp
@@ -32,27 +32,20 @@ void FAssetTypeActions_MaterialInterface::GetActions(const TArray<UObject*>& InO
 UThumbnailInfo* FAssetTypeActions_MaterialInterface::GetThumbnailInfo(UObject* Asset) const
 {
 	UMaterialInterface* MaterialInterface = CastChecked<UMaterialInterface>(Asset);
-	UThumbnailInfo* ThumbnailInfo = MaterialInterface->ThumbnailInfo;
-	if ( ThumbnailInfo == NULL )
+	USceneThumbnailInfoWithPrimitive* ThumbnailInfo = Cast<USceneThumbnailInfoWithPrimitive>(MaterialInterface->ThumbnailInfo);
+	if ( ThumbnailInfo == nullptr )
 	{
 		ThumbnailInfo = NewObject<USceneThumbnailInfoWithPrimitive>(MaterialInterface, NAME_None, RF_Transactional);
 		MaterialInterface->ThumbnailInfo = ThumbnailInfo;
 	}
-
-	return ThumbnailInfo;
-}
-
-EThumbnailPrimType FAssetTypeActions_MaterialInterface::GetDefaultThumbnailPrimitiveType(UObject* Asset) const
-{
-	EThumbnailPrimType PrimType = TPT_Sphere;
-	UMaterialInterface* MaterialInterface = CastChecked<UMaterialInterface>(Asset);
-	UMaterial* Material = MaterialInterface->GetBaseMaterial();
+	
+	const UMaterial* Material = MaterialInterface->GetBaseMaterial();
 	if (Material && Material->bUsedWithParticleSprites)
-	{
-		PrimType = TPT_Plane;
-	}
+    {
+    	ThumbnailInfo->DefaultPrimitiveType = TPT_Plane;
+    }
 
-	return PrimType;
+	return ThumbnailInfo;
 }
 
 void FAssetTypeActions_MaterialInterface::ExecuteNewMIC(TArray<TWeakObjectPtr<UMaterialInterface>> Objects)
diff --git a/Engine/Source/Developer/AssetTools/Private/AssetTypeActions/AssetTypeActions_MaterialInterface.h b/Engine/Source/Developer/AssetTools/Private/AssetTypeActions/AssetTypeActions_MaterialInterface.h
--- a/Engine/Source/Developer/AssetTools/Private/AssetTypeActions/AssetTypeActions_MaterialInterface.h
+++ b/Engine/Source/Developer/AssetTools/Private/AssetTypeActions/AssetTypeActions_MaterialInterface.h
@@ -18,7 +18,6 @@ public:
 	virtual bool CanFilter() override { return false; }
 	virtual uint32 GetCategories() override { return EAssetTypeCategories::Materials; }
 	virtual class UThumbnailInfo* GetThumbnailInfo(UObject* Asset) const override;
-	virtual EThumbnailPrimType GetDefaultThumbnailPrimitiveType(UObject* Asset) const override;
 
 private:
 	/** Handler for when NewMIC is selected */
diff --git a/Engine/Source/Developer/AssetTools/Private/AssetTypeActions/AssetTypeActions_PhysicalMaterialMask.cpp b/Engine/Source/Developer/AssetTools/Private/AssetTypeActions/AssetTypeActions_PhysicalMaterialMask.cpp
--- a/Engine/Source/Developer/AssetTools/Private/AssetTypeActions/AssetTypeActions_PhysicalMaterialMask.cpp
+++ b/Engine/Source/Developer/AssetTools/Private/AssetTypeActions/AssetTypeActions_PhysicalMaterialMask.cpp
@@ -76,17 +76,6 @@ void FAssetTypeActions_PhysicalMaterialMask::GetResolvedSourceFilePaths(const TA
 	}
 }
 
-
-class UThumbnailInfo* FAssetTypeActions_PhysicalMaterialMask::GetThumbnailInfo(UObject* Asset) const 
-{
-	return nullptr;
-}
-
-EThumbnailPrimType FAssetTypeActions_PhysicalMaterialMask::GetDefaultThumbnailPrimitiveType(UObject* Asset) const 
-{
-	return TPT_Sphere;
-}
-
 bool FAssetTypeActions_PhysicalMaterialMask::CanExecuteImportedAssetActions(const TArray<FString> ResolvedFilePaths) const
 {
 	// Verify that all the file paths are legitimate
diff --git a/Engine/Source/Developer/AssetTools/Private/AssetTypeActions/AssetTypeActions_PhysicalMaterialMask.h b/Engine/Source/Developer/AssetTools/Private/AssetTypeActions/AssetTypeActions_PhysicalMaterialMask.h
--- a/Engine/Source/Developer/AssetTools/Private/AssetTypeActions/AssetTypeActions_PhysicalMaterialMask.h
+++ b/Engine/Source/Developer/AssetTools/Private/AssetTypeActions/AssetTypeActions_PhysicalMaterialMask.h
@@ -17,12 +17,6 @@ public:
 	virtual bool IsImportedAsset() const override { return true; }
 	virtual void GetResolvedSourceFilePaths(const TArray<UObject*>& TypeAssets, TArray<FString>& OutSourceFilePaths) const override;
 
-	/** Returns the thumbnail info for the specified asset, if it has one. */
-	virtual class UThumbnailInfo* GetThumbnailInfo(UObject* Asset) const override;
-
-	/** Returns the default thumbnail type that should be rendered when rendering primitive shapes.  This does not need to be implemented if the asset does not render a primitive shape */
-	virtual EThumbnailPrimType GetDefaultThumbnailPrimitiveType(UObject* Asset) const override;
-
 private:
 	/** Handler to check to see if imported asset actions are allowed */
 	bool CanExecuteImportedAssetActions(const TArray<FString> ResolvedFilePaths) const;
diff --git a/Engine/Source/Developer/AssetTools/Private/AssetTypeActions/AssetTypeActions_Redirector.cpp b/Engine/Source/Developer/AssetTools/Private/AssetTypeActions/AssetTypeActions_Redirector.cpp
deleted file mode 100644
diff --git a/Engine/Source/Developer/AssetTools/Private/AssetTypeActions/AssetTypeActions_Redirector.h b/Engine/Source/Developer/AssetTools/Private/AssetTypeActions/AssetTypeActions_Redirector.h
deleted file mode 100644
diff --git a/Engine/Source/Developer/AssetTools/Private/AssetTypeActions/AssetTypeActions_World.cpp b/Engine/Source/Developer/AssetTools/Private/AssetTypeActions/AssetTypeActions_World.cpp
deleted file mode 100644
diff --git a/Engine/Source/Developer/AssetTools/Private/AssetTypeActions/AssetTypeActions_World.h b/Engine/Source/Developer/AssetTools/Private/AssetTypeActions/AssetTypeActions_World.h
deleted file mode 100644
diff --git a/Engine/Source/Developer/AssetTools/Public/AssetTypeActions/AssetTypeActions_Blueprint.h b/Engine/Source/Developer/AssetTools/Public/AssetTypeActions/AssetTypeActions_Blueprint.h
--- a/Engine/Source/Developer/AssetTools/Public/AssetTypeActions/AssetTypeActions_Blueprint.h
+++ b/Engine/Source/Developer/AssetTools/Public/AssetTypeActions/AssetTypeActions_Blueprint.h
@@ -18,6 +18,12 @@ public:
 	virtual FText GetName() const override { return NSLOCTEXT("AssetTypeActions", "AssetTypeActions_Blueprint", "Blueprint Class"); }
 	virtual FColor GetTypeColor() const override { return FColor( 63, 126, 255 ); }
 	virtual UClass* GetSupportedClass() const override { return UBlueprint::StaticClass(); }
+
+	// AssetDefinition - The Blueprint actions are now in the UAssetDefinition_Blueprint,
+	// can't use false here, since FAssetTypeActions_Blueprint is inherited.  Need the others to function until they can
+	// move.
+	virtual bool ShouldCallGetActions() const override { return GetSupportedClass() != UBlueprint::StaticClass(); }
+
 	virtual void GetActions(const TArray<UObject*>& InObjects, FToolMenuSection& Section) override;
 	virtual void OpenAssetEditor(const TArray<UObject*>& InObjects, TSharedPtr<class IToolkitHost> EditWithinLevelEditor = TSharedPtr<IToolkitHost>()) override;
 	virtual bool CanMerge() const override;
diff --git a/Engine/Source/Developer/AssetTools/Public/AssetTypeActions_Base.h b/Engine/Source/Developer/AssetTools/Public/AssetTypeActions_Base.h
--- a/Engine/Source/Developer/AssetTools/Public/AssetTypeActions_Base.h
+++ b/Engine/Source/Developer/AssetTools/Public/AssetTypeActions_Base.h
@@ -161,16 +161,6 @@ public:
 		return FText::GetEmpty();
 	}
 
-	virtual void SetSupported(bool bInSupported) final
-	{
-		bIsSupported = bInSupported;
-	}
-
-	virtual bool IsSupported() const final
-	{
-		return bIsSupported;
-	}
-
 	virtual bool SupportsOpenedMethod(const EAssetTypeActivationOpenedMethod OpenedMethod) const override
 	{
 		return OpenedMethod == EAssetTypeActivationOpenedMethod::Edit;
@@ -242,7 +232,4 @@ protected:
 
 		return TypedObjects;
 	}
-
-private:
-	bool bIsSupported = true;
 };
diff --git a/Engine/Source/Developer/AssetTools/Public/IAssetTools.h b/Engine/Source/Developer/AssetTools/Public/IAssetTools.h
--- a/Engine/Source/Developer/AssetTools/Public/IAssetTools.h
+++ b/Engine/Source/Developer/AssetTools/Public/IAssetTools.h
@@ -250,6 +250,10 @@ public:
 	/** Gets the appropriate AssetTypeActions for the supplied class */
 	virtual TWeakPtr<IAssetTypeActions> GetAssetTypeActionsForClass(const UClass* Class) const = 0;
 
+	virtual bool CanLocalize(const UClass* Class) const = 0;
+
+	virtual TOptional<FLinearColor> GetTypeColor(const UClass* Class) const = 0;
+
 	/** Gets the list of appropriate AssetTypeActions for the supplied class */
 	virtual TArray<TWeakPtr<IAssetTypeActions>> GetAssetTypeActionsListForClass(const UClass* Class) const = 0;
 
@@ -531,7 +535,7 @@ public:
 	virtual TSharedRef<FNamePermissionList>& GetAssetClassPermissionList() = 0;
 
 	/** Get asset class permission list for content browser and other systems */
-	virtual TSharedRef<FPathPermissionList>& GetAssetClassPathPermissionList(EAssetClassAction AssetClassAction) = 0;
+	virtual const TSharedRef<FPathPermissionList>& GetAssetClassPathPermissionList(EAssetClassAction AssetClassAction) const = 0;
 
 	/** Which BlueprintTypes are allowed to be created. An empty list should allow everything. */
 	virtual TSet<EBlueprintType>& GetAllowedBlueprintTypes() = 0;
diff --git a/Engine/Source/Developer/AssetTools/Public/IAssetTypeActions.h b/Engine/Source/Developer/AssetTools/Public/IAssetTypeActions.h
--- a/Engine/Source/Developer/AssetTools/Public/IAssetTypeActions.h
+++ b/Engine/Source/Developer/AssetTools/Public/IAssetTypeActions.h
@@ -7,7 +7,7 @@
 #include "Merge.h"
 #include "ThumbnailRendering/ThumbnailManager.h"
 #include "AssetRegistry/AssetData.h"
-#include "IAssetTypeActions.generated.h"
+#include "AssetDefinition.h"
 
 class IToolkitHost;
 enum class EAssetTypeActivationOpenedMethod : uint8;
@@ -22,39 +22,6 @@ namespace EAssetTypeActivationMethod
 	};
 }
 
-/* Revision information for a single revision of a file in source control */
-USTRUCT(BlueprintType)
-struct FRevisionInfo
-{
-	GENERATED_BODY()
-
-	UPROPERTY(BlueprintReadWrite, Category="Asset Revision")
-	FString		Revision;
-
-	UPROPERTY(BlueprintReadWrite, Category="Asset Revision")
-	int32		Changelist = -1;
-
-	UPROPERTY(BlueprintReadWrite, Category="Asset Revision")
-	FDateTime	Date;
-
-	static inline FRevisionInfo InvalidRevision()
-	{
-		static const FRevisionInfo Ret = { TEXT(""), -1, FDateTime() };
-		return Ret;
-	}
-};
-
-#ifndef DEPRECATE_ASSET_TYPE_ACTIONS_NEEDING_LOADED_OBJECTS
-	#define DEPRECATE_ASSET_TYPE_ACTIONS_NEEDING_LOADED_OBJECTS 0
-#endif
-
-#ifndef DEPRECATE_ASSET_TYPE_ACTIONS_CALLING_GETACTIONS
-    #if DEPRECATE_ASSET_TYPE_ACTIONS_NEEDING_LOADED_OBJECTS
-	    #define DEPRECATE_ASSET_TYPE_ACTIONS_CALLING_GETACTIONS DEPRECATE_ASSET_TYPE_ACTIONS_NEEDING_LOADED_OBJECTS
-	#else
-	    #define DEPRECATE_ASSET_TYPE_ACTIONS_CALLING_GETACTIONS 0
-	#endif
-#endif
 
 /** AssetTypeActions provide actions and other information about asset types */
 class IAssetTypeActions : public TSharedFromThis<IAssetTypeActions>
@@ -63,6 +30,12 @@ public:
 	/** Virtual destructor */
 	virtual ~IAssetTypeActions(){}
 
+	/**
+	 * Returns if this is actually just an UAssetDefinition.  Which doesn't exist in this module.  UAssetDefinition
+	 * is replacing IAssetTypeActions, but for now we generate fake IAssetTypeActions that are secretly AssetDefinitions
+	 */
+	virtual bool IsAssetDefinitionInDisguise() const { return false; }
+
 	/** Returns the name of this type */
 	virtual FText GetName() const = 0;
 
@@ -71,26 +44,18 @@ public:
 
 	/** Returns the color associated with this type */
 	virtual FColor GetTypeColor() const = 0;
+
+	/** Returns true if we should call GetActions. */
+	virtual bool ShouldCallGetActions() const { return true; }
 	
 	/** Returns true if this class can supply actions for InObjects. */
-//   	UE_DEPRECATED(5.1, "HasActions is no longer used all IAssetTypeActions are assumed to need to have them and GetActions is just called instead.  GetActions is always called unless DEPRECATE_ASSET_TYPE_ACTIONS_CALLING_GETACTIONS is true.")
    	virtual bool HasActions( const TArray<UObject*>& InObjects ) const { return true; }
 
-#if DEPRECATE_ASSET_TYPE_ACTIONS_NEEDING_LOADED_OBJECTS
-	/** Generates a menubuilder for the specified objects. */
-	UE_DEPRECATED(5.1, "Look at 'HOW TO ADD ASSET EXTENSION MENUS' in IAssetTypeActions.h")
-	virtual void GetActions( const TArray<UObject*>& InObjects, class FMenuBuilder& MenuBuilder ) final { }
-
-	/** Generates a menu section for the specified objects. */
-	UE_DEPRECATED(5.1, "Look at 'HOW TO ADD ASSET EXTENSION MENUS' in IAssetTypeActions.h")
-	virtual void GetActions(const TArray<UObject*>& InObjects, struct FToolMenuSection& Section) final { }
-#else
 	/** Generates a menubuilder for the specified objects. */
 	virtual void GetActions( const TArray<UObject*>& InObjects, class FMenuBuilder& MenuBuilder ) { }
 
 	/** Generates a menu section for the specified objects. */
 	virtual void GetActions(const TArray<UObject*>& InObjects, struct FToolMenuSection& Section) { }
-#endif
 
 	/** Opens the asset editor for the specified objects. If EditWithinLevelEditor is valid, the world-centric editor will be used. */
 	virtual void OpenAssetEditor( const TArray<UObject*>& InObjects, TSharedPtr<IToolkitHost> EditWithinLevelEditor = TSharedPtr<IToolkitHost>() ) = 0;
@@ -147,6 +112,7 @@ public:
 	virtual class UThumbnailInfo* GetThumbnailInfo(UObject* Asset) const = 0;
 
 	/** Returns the default thumbnail type that should be rendered when rendering primitive shapes.  This does not need to be implemented if the asset does not render a primitive shape */
+	UE_DEPRECATED(5.2, "When overriding GetThumbnailInfo, The default thumbnail primitive is now a property on USceneThumbnailInfoWithPrimitive.  The default only applies to this thumbnail type, so it has been relocated as an implementation detail of it.")
 	virtual EThumbnailPrimType GetDefaultThumbnailPrimitiveType(UObject* Asset) const = 0;
 
 	/** Optionally returns a custom widget to overlay on top of this assets' thumbnail */
@@ -171,12 +137,15 @@ public:
 	virtual FText GetDisplayNameFromAssetData(const FAssetData& AssetData) const = 0;
 
 	/** Sets whether or not this asset type is a supported type for this editor session. */
-	virtual void SetSupported(bool bInSupported) = 0;
+	UE_DEPRECATED(5.2, "Use IAssetTools::IsAssetClassSupported")
+	virtual void SetSupported(bool bInSupported) final { ensureMsgf(false, TEXT("Deprecated")); };
 
 	/** Is this asset type supported in the current session? */
-	virtual bool IsSupported() const = 0;
-
+	UE_DEPRECATED(5.2, "Use IAssetTools::IsAssetClassSupported")
+	virtual bool IsSupported() const final { ensureMsgf(false, TEXT("Deprecated")); return true; }
+	
 	/** Returns class path name as a package + class FName pair */
+	//UE_DEPRECATED(5.2, "Just call GetSupportedClass()->GetClassPathName();")
 	virtual FTopLevelAssetPath GetClassPathName() const = 0;
 
 	/** Does this asset support edit or view methods? */
diff --git a/Engine/Source/Editor/AssetDefinition/AssetDefinition.Build.cs b/Engine/Source/Editor/AssetDefinition/AssetDefinition.Build.cs
new file mode 100644
--- /dev/null
+++ b/Engine/Source/Editor/AssetDefinition/AssetDefinition.Build.cs
@@ -0,0 +1,33 @@
+// Copyright Epic Games, Inc. All Rights Reserved.
+
+using UnrealBuildTool;
+using System.IO;
+
+public class AssetDefinition : ModuleRules
+{
+	public AssetDefinition(ReadOnlyTargetRules Target) : base(Target)
+	{		
+        PublicDependencyModuleNames.AddRange(
+			new string[] { 
+				"Core",
+				"CoreUObject",
+				"AssetRegistry",
+				"Engine",
+                "InputCore",
+				"Slate",
+				"SlateCore",
+				"EditorSubsystem"
+			}
+		);
+
+		PrivateDependencyModuleNames.AddRange(
+			new string[] {
+			}
+		);
+		
+		PrivateIncludePaths.AddRange(
+        	new string[] {
+        	}
+        );
+	}
+}
diff --git a/Engine/Source/Editor/AssetDefinition/Private/AssetDefinition.cpp b/Engine/Source/Editor/AssetDefinition/Private/AssetDefinition.cpp
new file mode 100644
--- /dev/null
+++ b/Engine/Source/Editor/AssetDefinition/Private/AssetDefinition.cpp
@@ -0,0 +1,83 @@
+// Copyright Epic Games, Inc. All Rights Reserved.
+
+#include "AssetDefinition.h"
+#include "AssetDefinitionRegistry.h"
+#include "Misc/AssetFilterData.h"
+
+#define LOCTEXT_NAMESPACE "UAssetDefinition"
+
+FAssetCategoryPath EAssetCategoryPaths::Basic(LOCTEXT("Basic", "Basic"));
+FAssetCategoryPath EAssetCategoryPaths::Animation(LOCTEXT("Animation", "Animation"));
+FAssetCategoryPath EAssetCategoryPaths::Material(LOCTEXT("Material", "Material"));
+FAssetCategoryPath EAssetCategoryPaths::Audio(LOCTEXT("Audio", "Audio"));
+FAssetCategoryPath EAssetCategoryPaths::Physics(LOCTEXT("Physics", "Physics"));
+FAssetCategoryPath EAssetCategoryPaths::UI(LOCTEXT("UserInterface", "User Interface"));
+FAssetCategoryPath EAssetCategoryPaths::Misc(LOCTEXT("Miscellaneous", "Miscellaneous"));
+FAssetCategoryPath EAssetCategoryPaths::Gameplay(LOCTEXT("Gameplay", "Gameplay"));
+FAssetCategoryPath EAssetCategoryPaths::Blueprint(LOCTEXT("Blueprint", "Blueprint"));
+FAssetCategoryPath EAssetCategoryPaths::Texture(LOCTEXT("Texture", "Texture"));
+
+FAssetCategoryPath::FAssetCategoryPath(const FText InCategory)
+{
+	CategoryPath = { TPair<FName, FText>(FName(*FTextInspector::GetSourceString(InCategory)), InCategory) };
+}
+
+FAssetCategoryPath::FAssetCategoryPath(TConstArrayView<FText> InCategoryPath)
+{
+	check(InCategoryPath.Num() > 0);
+	
+	for (const FText& CategoryChunk : InCategoryPath)
+	{
+		CategoryPath.Add(TPair<FName, FText>(FName(*FTextInspector::GetSourceString(CategoryChunk)), CategoryChunk));
+	}
+}
+
+UAssetDefinition::UAssetDefinition()
+{
+}
+
+void UAssetDefinition::PostCDOContruct()
+{
+	Super::PostCDOContruct();
+
+	if (CanRegisterStatically())
+	{
+		UAssetDefinitionRegistry::Get()->RegisterAssetDefinition(this);
+	}
+}
+
+bool UAssetDefinition::CanRegisterStatically() const
+{
+	return !GetClass()->HasAnyClassFlags(CLASS_Abstract);
+}
+
+EAssetCommandResult UAssetDefinition::GetFilters(TArray<FAssetFilterData>& OutFilters) const
+{
+	const TSoftClassPtr<UObject> AssetClassPtr = GetAssetClass();
+
+	if (const UClass* AssetClass = AssetClassPtr.Get())
+	{
+		// By default we don't advertise filtering if the class is abstract for the asset definition.  Odds are,
+		// if they've registered an abstract class as an asset definition, they mean to use it for subclasses.
+		if (!AssetClass->HasAnyClassFlags(CLASS_Abstract))
+		{
+			FAssetFilterData DefaultFilter;
+			DefaultFilter.Name = AssetClassPtr.ToSoftObjectPath().ToString();
+			DefaultFilter.DisplayText = GetAssetDisplayName();
+			DefaultFilter.Filter.ClassPaths.Add(AssetClassPtr.ToSoftObjectPath().GetAssetPath());
+			DefaultFilter.Filter.bRecursiveClasses = true;
+			OutFilters.Add(MoveTemp(DefaultFilter));
+	
+			return EAssetCommandResult::Handled;
+		}
+	}
+	
+	return EAssetCommandResult::Unhandled;
+}
+
+TConstArrayView<FAssetCategoryPath> UAssetDefinition::GetAssetCategories() const
+{
+	return { EAssetCategoryPaths::Basic };
+}
+
+#undef LOCTEXT_NAMESPACE
\ No newline at end of file
diff --git a/Engine/Source/Editor/AssetDefinition/Private/AssetDefinitionModule.cpp b/Engine/Source/Editor/AssetDefinition/Private/AssetDefinitionModule.cpp
new file mode 100644
--- /dev/null
+++ b/Engine/Source/Editor/AssetDefinition/Private/AssetDefinitionModule.cpp
@@ -0,0 +1,23 @@
+﻿// Copyright Epic Games, Inc. All Rights Reserved.
+
+#include "Modules/ModuleManager.h"
+
+#define LOCTEXT_NAMESPACE "FAssetDefinitionModule"
+
+class FAssetDefinitionModule : public IModuleInterface
+{
+public:
+
+	// IModuleInterface interface
+	virtual void StartupModule() override
+	{
+	}
+
+	virtual void ShutdownModule() override
+	{
+	}
+};
+
+#undef LOCTEXT_NAMESPACE
+
+IMPLEMENT_MODULE(FAssetDefinitionModule, AssetDefinition);
diff --git a/Engine/Source/Editor/AssetDefinition/Private/AssetDefinitionRegistry.cpp b/Engine/Source/Editor/AssetDefinition/Private/AssetDefinitionRegistry.cpp
new file mode 100644
--- /dev/null
+++ b/Engine/Source/Editor/AssetDefinition/Private/AssetDefinitionRegistry.cpp
@@ -0,0 +1,101 @@
+// Copyright Epic Games, Inc. All Rights Reserved.
+
+#include "AssetDefinitionRegistry.h"
+#include "AssetDefinition.h"
+#include "Algo/Accumulate.h"
+#include "AssetRegistry/AssetRegistryHelpers.h"
+#include "Engine/AssetManager.h"
+
+UAssetDefinitionRegistry* UAssetDefinitionRegistry::Singleton = nullptr;
+bool UAssetDefinitionRegistry::bHasShutDown = false;
+
+UAssetDefinitionRegistry* UAssetDefinitionRegistry::Get()
+{
+	if (!Singleton && !bHasShutDown)
+	{
+		// Required for StartupModule and ShutdownModule to be called and FindModule to list the AssetDefinition module
+		FModuleManager::Get().LoadModuleChecked("AssetDefinition");
+
+		Singleton = NewObject<UAssetDefinitionRegistry>();
+		Singleton->AddToRoot();
+		check(Singleton);
+	}
+
+	return Singleton;
+}
+
+void UAssetDefinitionRegistry::BeginDestroy()
+{
+	if (Singleton == this)
+	{
+		bHasShutDown = true;
+		Singleton = nullptr;
+	}
+
+	Super::BeginDestroy();
+}
+
+const UAssetDefinition* UAssetDefinitionRegistry::GetAssetDefinitionForAsset(const FAssetData& Asset) const
+{
+	const UClass* AssetClass = UAssetRegistryHelpers::FindAssetNativeClass(Asset);
+	return GetAssetDefinitionForClass(AssetClass);
+}
+
+const UAssetDefinition* UAssetDefinitionRegistry::GetAssetDefinitionForClass(const UClass* Class) const
+{
+	while(Class)
+	{
+		if (const UAssetDefinition* AssetDefinition = AssetDefinitions.FindRef(Class))
+		{
+			return AssetDefinition;
+		}
+		
+		Class = Class->GetSuperClass();
+	}
+	
+	return nullptr;
+}
+
+TArray<TObjectPtr<UAssetDefinition>> UAssetDefinitionRegistry::GetAllAssetDefinitions() const
+{
+	TArray<TObjectPtr<UAssetDefinition>> AllAssetDefinitions;
+	AssetDefinitions.GenerateValueArray(AllAssetDefinitions);
+	return AllAssetDefinitions;
+}
+
+void UAssetDefinitionRegistry::RegisterAssetDefinition(UAssetDefinition* AssetDefinition)
+{
+	check(AssetDefinition);
+	
+	if (TSoftClassPtr<UObject> SupportedClass = AssetDefinition->GetAssetClass())
+	{
+		const bool CanAddDefinition =
+			!AssetDefinitions.Contains(SupportedClass) ||
+			// This hack is basically to ensure that the statically created AssetDefinitions for things like blueprints
+			// will work alongside the AssetTypeActions for blueprints until that class can be completely subsumbed.
+			(AssetDefinition->CanRegisterStatically() && !AssetDefinitions[SupportedClass]->CanRegisterStatically());
+		
+		// Normally I'd ensure here, we ignore duplicates, but the IAssetTypeActions allowed duplicates.
+		if (CanAddDefinition)
+		{
+			AssetDefinitions.Add(SupportedClass, AssetDefinition);
+			AssetDefinition->OnRegistered();
+		}
+		else
+		{
+			UAssetDefinition* ExistingDefinition = AssetDefinitions[SupportedClass];
+			UE_LOG(LogTemp, Warning, TEXT("Unable to register %s for (%s), %s is already registered"),  *GetPathNameSafe(AssetDefinition), *SupportedClass.ToString(), *GetPathNameSafe(ExistingDefinition));
+		}
+	}
+}
+
+void UAssetDefinitionRegistry::UnregisterAssetDefinition(UAssetDefinition* AssetDefinition)
+{
+	check(AssetDefinition);
+	
+	if (TSoftClassPtr<UObject> SupportedClass = AssetDefinition->GetAssetClass())
+	{
+		AssetDefinitions.Remove(SupportedClass);
+		AssetDefinition->OnUnregistered();
+	}
+}
\ No newline at end of file
diff --git a/Engine/Source/Editor/AssetDefinition/Public/AssetDefinition.h b/Engine/Source/Editor/AssetDefinition/Public/AssetDefinition.h
new file mode 100644
--- /dev/null
+++ b/Engine/Source/Editor/AssetDefinition/Public/AssetDefinition.h
@@ -0,0 +1,395 @@
+// Copyright Epic Games, Inc. All Rights Reserved.
+
+#pragma once
+
+#include "UObject/Object.h"
+#include "UObject/SoftObjectPtr.h"
+#include "AssetRegistry/AssetData.h"
+#include "Toolkits/IToolkit.h"
+
+#include "AssetDefinition.generated.h"
+
+struct FAssetFilterData;
+class IToolkitHost;
+class UThumbnailInfo;
+struct FSlateBrush;
+class ISourceControlRevision;
+class SWidget;
+
+UENUM()
+enum class EAssetActivationMethod : uint8
+{
+	DoubleClicked,
+	Opened,
+	Previewed
+};
+
+UENUM()
+enum class EAssetCommandResult : uint8
+{
+	Handled,
+	Unhandled
+};
+
+UENUM()
+enum class EAssetOpenMethod : uint8
+{
+	Edit,
+	View,
+	//Preview
+};
+
+enum class EAssetMergeResult : uint8
+{
+	Unknown,
+	Completed,
+	Cancelled,
+};
+
+struct FAssetArgs
+{
+	FAssetArgs() { }
+	FAssetArgs(TConstArrayView<FAssetData> InAssets) : Assets(InAssets) { }
+	
+	TConstArrayView<FAssetData> Assets;
+
+	template<typename ExpectedObjectType>
+	TArray<ExpectedObjectType*> LoadObjects(const TSet<FName>& LoadTags = {}) const
+	{
+		TArray<ExpectedObjectType*> LoadedObjects;
+		LoadedObjects.Reserve(Assets.Num());
+		
+		for (const FAssetData& Asset : Assets)
+		{
+			if (Asset.IsInstanceOf(ExpectedObjectType::StaticClass()))
+			{
+				if (ExpectedObjectType* ExpectedType = Cast<ExpectedObjectType>(Asset.GetAsset(LoadTags)))
+				{
+					LoadedObjects.Add(ExpectedType);
+				}
+			}	
+		}
+		
+		return LoadedObjects;
+	}
+};
+
+struct FAssetOpenArgs : public FAssetArgs
+{
+	EAssetOpenMethod OpenMethod;
+	TSharedPtr<IToolkitHost> ToolkitHost;
+};
+
+struct FAssetActivateArgs : public FAssetArgs
+{
+	EAssetActivationMethod ActivationMethod;
+};
+
+struct FAssetSourceFileArgs : public FAssetArgs
+{
+	FAssetSourceFileArgs(TConstArrayView<FAssetData> InAssets) : FAssetArgs(InAssets) { }
+};
+
+struct FAssetSourceFile
+{
+	FString DisplayLabelName;
+	FString RelativeFilename;
+};
+
+struct FAssetMergeResults
+{
+	UPackage* MergedPackage = nullptr;
+	EAssetMergeResult Result = EAssetMergeResult::Unknown;
+};
+
+DECLARE_DELEGATE_OneParam(FOnAssetMergeResolved, const FAssetMergeResults& Results);
+
+struct FAssetMergeArgs
+{
+	FAssetData LocalAsset;
+	TOptional<FAssetData> BaseAsset;
+	TOptional<FAssetData> RemoteAsset;
+	
+	FOnAssetMergeResolved ResolutionCallback;
+};
+
+struct FAssetSupportResponse
+{
+public:
+	static FAssetSupportResponse Supported()
+	{
+		return FAssetSupportResponse(true, FText::GetEmpty());
+	}
+
+	static FAssetSupportResponse NotSupported()
+	{
+		return FAssetSupportResponse(false, FText::GetEmpty());
+	}
+
+	static FAssetSupportResponse Error(const FText& ErrorText)
+	{
+		return FAssetSupportResponse(false, ErrorText);
+	}
+
+	bool IsSupported() const { return bSupported; }
+	const FText& GetErrorText() const { return ErrorText; }
+
+private:
+	FAssetSupportResponse(bool InSupported, const FText InError)
+		: bSupported(InSupported)
+		, ErrorText(InError)
+	{
+	}
+
+private:
+	bool bSupported;
+	FText ErrorText;
+};
+
+/* Revision information for a single revision of a file in source control */
+USTRUCT(BlueprintType)
+struct FRevisionInfo
+{
+	GENERATED_BODY()
+
+	UPROPERTY(BlueprintReadWrite, Category="Asset Revision")
+	FString		Revision;
+
+	UPROPERTY(BlueprintReadWrite, Category="Asset Revision")
+	int32		Changelist = -1;
+
+	UPROPERTY(BlueprintReadWrite, Category="Asset Revision")
+	FDateTime	Date;
+
+	static inline FRevisionInfo InvalidRevision()
+	{
+		static const FRevisionInfo Ret = { TEXT(""), -1, FDateTime() };
+		return Ret;
+	}
+};
+
+struct FAssetDiffArgs
+{
+	UObject* OldAsset = nullptr;
+	FRevisionInfo OldRevision;
+
+	UObject* NewAsset = nullptr;
+	FRevisionInfo NewRevision;
+};
+
+struct FAssetOpenSupportArgs
+{
+	EAssetOpenMethod OpenMethod;
+};
+
+struct ASSETDEFINITION_API FAssetCategoryPath
+{
+	FAssetCategoryPath(FText InCategory);
+	FAssetCategoryPath(TConstArrayView<FText> InCategoryPath);
+
+	FName GetCategory() const { return CategoryPath[0].Key; }
+	FText GetCategoryText() const { return CategoryPath[0].Value; }
+	
+	bool HasSubCategory() const { return CategoryPath.Num() > 1; }
+	FName GetSubCategory() const { return HasSubCategory() ? CategoryPath[1].Key : NAME_None; }
+	FText GetSubCategoryText() const { return HasSubCategory() ? CategoryPath[1].Value : FText::GetEmpty(); }
+
+private:
+	TArray<TPair<FName, FText>> CategoryPath;
+};
+
+struct ASSETDEFINITION_API EAssetCategoryPaths
+{
+	static FAssetCategoryPath Basic;
+	static FAssetCategoryPath Animation;
+	static FAssetCategoryPath Material;
+	static FAssetCategoryPath Audio;
+	static FAssetCategoryPath Physics;
+	static FAssetCategoryPath UI;
+	static FAssetCategoryPath Misc;
+	static FAssetCategoryPath Gameplay;
+	static FAssetCategoryPath Blueprint;
+	static FAssetCategoryPath Texture;
+};
+
+struct FAssetOpenSupport
+{
+	FAssetOpenSupport(EAssetOpenMethod InOpenMethod, bool bInSupported)
+		: OpenMethod(InOpenMethod)
+		, IsSupported(bInSupported)
+	{
+	}
+	
+	EAssetOpenMethod OpenMethod;
+	bool IsSupported;
+	TOptional<EToolkitMode::Type> RequiredToolkitMode;
+};
+
+
+class UAssetDefinitionRegistry;
+
+/**
+ * Asset Definitions represent top level assets that are known to the editor.
+ */
+UCLASS(Abstract)
+class ASSETDEFINITION_API UAssetDefinition : public UObject
+{
+	GENERATED_BODY()
+
+public:
+	UAssetDefinition();
+
+	//Begin UObject
+	virtual void PostCDOContruct() override;
+	//End UObject
+	
+public:
+	
+	/** Returns the name of this type */
+	virtual FText GetAssetDisplayName() const PURE_VIRTUAL(UAssetDefinition::GetAssetDisplayName, ensureMsgf(false && "NotImplemented", TEXT("NotImplemented")); return FText(); );
+
+	/**
+	 * Returns the name of this type, but allows overriding the default on a specific instance of the asset.  This
+	 * is handy for cases like UAssetData which are of course all UAssetData, but a given instance of the asset
+	 * is really a specific instance of some UAssetData class, and being able to override that on the instance is handy
+	 * for readability at the Content Browser level.
+	 */
+	virtual FText GetAssetDisplayName(const FAssetData& AssetData) const { return GetAssetDisplayName(); }
+
+	/** Get the supported class of this type. */
+	virtual TSoftClassPtr<UObject> GetAssetClass() const PURE_VIRTUAL(UAssetDefinition::GetAssetClass, ensureMsgf(false && "NotImplemented", TEXT("NotImplemented")); return TSoftClassPtr<UClass>(); );
+
+	/** Returns the color associated with this type */
+	virtual FLinearColor GetAssetColor() const PURE_VIRTUAL(UAssetDefinition::GetAssetColor, ensureMsgf(false && "NotImplemented", TEXT("NotImplemented")); return FColor::Red; );
+
+	/** Returns additional tooltip information for the specified asset, if it has any. */
+	virtual FText GetAssetDescription(const FAssetData& AssetData) const { return FText::GetEmpty(); }
+
+	/** Gets a list of categories this asset is in, these categories are used to help organize */
+	virtual TConstArrayView<FAssetCategoryPath> GetAssetCategories() const;
+	
+public:
+	// Common Operations
+	virtual TArray<FAssetData> PrepareToActivateAssets(const FAssetActivateArgs& ActivateArgs) const
+	{
+		return TArray<FAssetData>(ActivateArgs.Assets);
+	}
+	
+	/** Get open support for the method.  Includes required information before we call OpenAsset. */
+	virtual FAssetOpenSupport GetAssetOpenSupport(const FAssetOpenSupportArgs& OpenSupportArgs) const
+	{
+		return FAssetOpenSupport(OpenSupportArgs.OpenMethod,OpenSupportArgs.OpenMethod == EAssetOpenMethod::Edit); 
+	}
+	
+	virtual EAssetCommandResult OpenAssets(const FAssetOpenArgs& OpenArgs) const PURE_VIRTUAL(UAssetDefinition::OpenAsset, ensureMsgf(false && "NotImplemented", TEXT("NotImplemented")); return EAssetCommandResult::Unhandled; );
+
+	virtual EAssetCommandResult ActivateAssets(const FAssetActivateArgs& ActivateArgs) const 
+	{
+		return EAssetCommandResult::Unhandled;
+	}
+
+	
+	// Common Queries
+	virtual FAssetSupportResponse CanRename(const FAssetData& InAsset) const
+	{
+		return FAssetSupportResponse::Supported();
+	}
+
+	virtual FAssetSupportResponse CanDuplicate(const FAssetData& InAsset) const
+	{
+		return FAssetSupportResponse::Supported();
+	}
+
+	virtual FAssetSupportResponse CanLocalize(const FAssetData& InAsset) const
+	{
+		return FAssetSupportResponse::NotSupported();
+	}
+
+	
+	// Importing
+	virtual bool CanImport() const { return false; }
+
+	
+	// Merging
+	virtual bool CanMerge() const { return false; }
+	virtual EAssetCommandResult Merge(const FAssetMergeArgs& MergeArgs) const { return EAssetCommandResult::Unhandled; }
+	
+
+	// Filtering
+	virtual EAssetCommandResult GetFilters(TArray<FAssetFilterData>& OutFilters) const;
+
+	
+	// Extras
+	virtual FText GetObjectDisplayText(UObject* Object) const
+	{
+		return FText::FromString(Object->GetName());
+	}
+
+	// Source Files
+	virtual EAssetCommandResult GetSourceFiles(const FAssetSourceFileArgs& SourceFileArgs, TArray<FAssetSourceFile>& OutSourceAssets) const
+	{
+		return EAssetCommandResult::Unhandled;
+	}
+
+	// Diffing Assets
+	virtual EAssetCommandResult PerformAssetDiff(const FAssetDiffArgs& DiffArgs) const
+	{
+		return EAssetCommandResult::Unhandled;
+	}
+
+	// Thumbnails
+
+	/** Returns the thumbnail info for the specified asset, if it has one. This typically requires loading the asset.  */
+	virtual UThumbnailInfo* LoadThumbnailInfo(const FAssetData& Asset) const
+	{
+		return nullptr;
+	}
+
+	/** Returns thumbnail brush unique for the given asset data.  Returning null falls back to class thumbnail brush. */
+	virtual const FSlateBrush* GetThumbnailBrush(const FAssetData& InAssetData, const FName InClassName) const
+	{
+		return nullptr;
+	}
+
+	/** Returns icon brush unique for the given asset data.  Returning null falls back to class icon brush. */
+	virtual const FSlateBrush* GetIconBrush(const FAssetData& InAssetData, const FName InClassName) const
+	{
+		return nullptr;
+	}
+	
+	/** Optionally returns a custom widget to overlay on top of this assets' thumbnail */
+	virtual TSharedPtr<SWidget> GetThumbnailOverlay(const FAssetData& AssetData) const
+	{
+		return TSharedPtr<SWidget>();
+	}
+
+	// DEVELOPER NOTE:
+	// Originally this class was based on the IAssetTypeActions implementation.  Several of the functions on there
+	// were created organically and added without a larger discussion about if such a thing belonged on those classes.
+	//
+	// For example, IAssetTypeActions::ShouldForceWorldCentric was needed for a single asset, but we didn't instead
+	// implement GetAssetOpenSupport, which merges the needs of ShouldForceWorldCentric, and SupportsOpenedMethod.
+	// 
+	// Another example, is IAssetTypeActions::SetSupported and IAssetTypeActions::IsSupported.  These were concepts
+	// that could have lived in a map on the registry and never needed to be stored on the actual IAssetTypeActions.
+	//
+	// So, please do not add new functions to this class if it can be helped.  The AssetDefinitions are intended to be
+	// a basic low level representation of top level assets for the Content Browser and other editor tools to do
+	// some basic interaction with them, or learn some basic common details about them.
+	//
+	// If you must add a new function to this class, some requests,
+	// 1. Can it be added as a parameter to an existing Argument struct for an existing function?  If so, please do that.
+	// 2. Can it be added as part of the return structure of an existing function?  If so, please do that.
+	// 3. If you add a new function, please create a struct for the Args.  We'll be able to upgrade things easier.
+	//    Please continue to use EAssetCommandResult and FAssetSupportResponse, for those kinds of commands.
+
+public:
+
+	virtual void OnRegistered() { }
+	virtual void OnUnregistered() { }
+
+protected:
+	virtual bool CanRegisterStatically() const;
+
+	friend class UAssetDefinitionRegistry;
+};
\ No newline at end of file
diff --git a/Engine/Source/Editor/AssetDefinition/Public/AssetDefinitionRegistry.h b/Engine/Source/Editor/AssetDefinition/Public/AssetDefinitionRegistry.h
new file mode 100644
--- /dev/null
+++ b/Engine/Source/Editor/AssetDefinition/Public/AssetDefinitionRegistry.h
@@ -0,0 +1,42 @@
+// Copyright Epic Games, Inc. All Rights Reserved.
+
+#pragma once
+
+#include "UObject/Object.h"
+
+#include "AssetDefinitionRegistry.generated.h"
+
+class UAssetDefinition;
+
+UCLASS(config=Editor)
+class ASSETDEFINITION_API UAssetDefinitionRegistry : public UObject
+{
+	GENERATED_BODY()
+
+public:
+	static UAssetDefinitionRegistry* Get();
+	
+	virtual void BeginDestroy() override;
+
+	const UAssetDefinition* GetAssetDefinitionForAsset(const FAssetData& Asset) const;
+	const UAssetDefinition* GetAssetDefinitionForClass(const UClass* Class) const;
+
+	TArray<TObjectPtr<UAssetDefinition>> GetAllAssetDefinitions() const;
+
+	/**
+	 * Normally UAssetDefinitionRegistry are registered automatically by their CDO.  The only reason you need to do this is if
+	 * you're forced to dynamically create the UAssetDefinition at runtime.  The original reason for this function was
+	 * to be able to create wrappers for the to be replaced IAssetTypeActions, that you can access AssetDefinition
+	 * versions of any IAssetType making the upgrade easier.
+	 */
+	void RegisterAssetDefinition(UAssetDefinition* AssetDefinition);
+
+	void UnregisterAssetDefinition(UAssetDefinition* AssetDefinition);
+	
+private:
+	static UAssetDefinitionRegistry* Singleton;
+	static bool bHasShutDown;
+
+	UPROPERTY()
+	TMap<TSoftClassPtr<UObject>, TObjectPtr<UAssetDefinition>> AssetDefinitions;
+};
\ No newline at end of file
diff --git a/Engine/Source/Editor/AssetDefinition/Public/Misc/AssetFilterData.h b/Engine/Source/Editor/AssetDefinition/Public/Misc/AssetFilterData.h
new file mode 100644
--- /dev/null
+++ b/Engine/Source/Editor/AssetDefinition/Public/Misc/AssetFilterData.h
@@ -0,0 +1,19 @@
+// Copyright Epic Games, Inc. All Rights Reserved.
+
+#pragma once
+
+#include "UObject/Object.h"
+#include "AssetRegistry/ARFilter.h"
+
+#include "AssetFilterData.generated.h"
+
+USTRUCT()
+struct FAssetFilterData
+{
+	GENERATED_BODY()
+
+public:
+	FString Name;
+	FText DisplayText;	
+	FARFilter Filter;
+};
diff --git a/Engine/Source/Editor/ContentBrowser/ContentBrowser.Build.cs b/Engine/Source/Editor/ContentBrowser/ContentBrowser.Build.cs
--- a/Engine/Source/Editor/ContentBrowser/ContentBrowser.Build.cs
+++ b/Engine/Source/Editor/ContentBrowser/ContentBrowser.Build.cs
@@ -21,7 +21,8 @@ public class ContentBrowser : ModuleRules
 
 		PrivateDependencyModuleNames.AddRange(
 			new string[] {
-                "AppFramework",
+			    "AssetDefinition",
+				"AppFramework",
 				"Core",
 				"CoreUObject",
 				"ApplicationCore",
diff --git a/Engine/Source/Editor/ContentBrowser/Private/AssetContextMenu.cpp b/Engine/Source/Editor/ContentBrowser/Private/AssetContextMenu.cpp
--- a/Engine/Source/Editor/ContentBrowser/Private/AssetContextMenu.cpp
+++ b/Engine/Source/Editor/ContentBrowser/Private/AssetContextMenu.cpp
@@ -2,6 +2,7 @@
 
 #include "AssetContextMenu.h"
 
+#include "AssetDefinition.h"
 #include "AssetRegistry/AssetData.h"
 #include "AssetToolsModule.h"
 #include "AssetViewUtils.h"
@@ -61,6 +62,10 @@
 #include "UObject/UnrealNames.h"
 #include "UObject/WeakObjectPtr.h"
 #include "UObject/WeakObjectPtrTemplates.h"
+#include "AssetDefinitionRegistry.h"
+#include "AssetRegistry/AssetRegistryHelpers.h"
+#include "Engine/AssetManager.h"
+#include "Misc/WarnIfAssetsLoadedInScope.h"
 
 class FMenuBuilder;
 class SWidget;
@@ -94,6 +99,8 @@ void FAssetContextMenu::BindCommands(TSharedPtr< FUICommandList >& Commands)
 
 TSharedRef<SWidget> FAssetContextMenu::MakeContextMenu(TArrayView<const FContentBrowserItem> InSelectedItems, const FSourcesData& InSourcesData, TSharedPtr< FUICommandList > InCommandList)
 {
+	FWarnIfAssetsLoadedInScope WarnIfAssetsLoaded;
+	
 	SetSelectedItems(InSelectedItems);
 	SourcesData = InSourcesData;
 
@@ -195,12 +202,12 @@ TSharedRef<SWidget> FAssetContextMenu::MakeContextMenu(TArrayView<const FContent
 			{
 				if (CommonClass == nullptr)
 				{
-					CommonClass = SelectedAssets[ObjIdx].GetClass();
+					CommonClass = UAssetRegistryHelpers::FindAssetNativeClass(SelectedAssets[ObjIdx]);
 					continue;
 				}
 
-				// Update the CommonClass until we find a common shared class.
-				UClass* Class = SelectedAssets[ObjIdx].GetClass();
+				// Update the CommonClass until we find a common shared class, ignore anything that's not native.
+				UClass* Class = UAssetRegistryHelpers::FindAssetNativeClass(SelectedAssets[ObjIdx]);
 				while (Class && !Class->IsChildOf(CommonClass))
 				{
 					CommonClass = CommonClass->GetSuperClass();
@@ -233,12 +240,18 @@ TSharedRef<SWidget> FAssetContextMenu::MakeContextMenu(TArrayView<const FContent
 				}
 			}
 
-			MenuName = UToolMenus::JoinMenuPaths(BaseMenuName, CommonClass->GetFName());
+			if (ensure(CommonClass))
+			{
+				MenuName = UToolMenus::JoinMenuPaths(BaseMenuName, CommonClass->GetFName());
 
-			RegisterMenuHierarchy(CommonClass);
+				RegisterMenuHierarchy(CommonClass);
 
-			// Find asset actions for common class
-			ContextObject->CommonAssetTypeActions = AssetToolsModule.Get().GetAssetTypeActionsForClass(ContextObject->CommonClass).Pin();
+				// Find asset actions for common class
+				ContextObject->CommonAssetDefinition = UAssetDefinitionRegistry::Get()->GetAssetDefinitionForClass(ContextObject->CommonClass);
+				PRAGMA_DISABLE_DEPRECATION_WARNINGS
+				ContextObject->CommonAssetTypeActions = AssetToolsModule.Get().GetAssetTypeActionsForClass(ContextObject->CommonClass).Pin();
+				PRAGMA_ENABLE_DEPRECATION_WARNINGS
+			}
 		}
 		else if (SelectedAssets.Num() == 0)
 		{
@@ -330,29 +343,29 @@ void FAssetContextMenu::RegisterContextMenu(const FName MenuName)
 			Section.AddDynamicEntry("GetActions", FNewToolMenuSectionDelegate::CreateLambda([](FToolMenuSection& InSection)
 			{
 				UContentBrowserAssetContextMenuContext* Context = InSection.FindContext<UContentBrowserAssetContextMenuContext>();
-				if (Context && Context->CommonAssetTypeActions.IsValid())
+				PRAGMA_DISABLE_DEPRECATION_WARNINGS
+				if (Context && Context->CommonAssetTypeActions.IsValid() && Context->CommonAssetTypeActions.Pin()->ShouldCallGetActions())
 				{
-					PRAGMA_DISABLE_DEPRECATION_WARNINGS
 					Context->CommonAssetTypeActions.Pin()->GetActions(Context->LoadSelectedObjectsIfNeeded(), InSection);
-					PRAGMA_ENABLE_DEPRECATION_WARNINGS
 				}
+				PRAGMA_ENABLE_DEPRECATION_WARNINGS
 			}));
 
 			Section.AddDynamicEntry("GetActionsLegacy", FNewToolMenuDelegateLegacy::CreateLambda([](FMenuBuilder& MenuBuilder, UToolMenu* InMenu)
 			{
 				UContentBrowserAssetContextMenuContext* Context = InMenu->FindContext<UContentBrowserAssetContextMenuContext>();
-				if (Context && Context->CommonAssetTypeActions.IsValid())
+				PRAGMA_DISABLE_DEPRECATION_WARNINGS
+				if (Context && Context->CommonAssetTypeActions.IsValid() && Context->CommonAssetTypeActions.Pin()->ShouldCallGetActions())
 				{
-					PRAGMA_DISABLE_DEPRECATION_WARNINGS
 					Context->CommonAssetTypeActions.Pin()->GetActions(Context->LoadSelectedObjectsIfNeeded(), MenuBuilder);
-					PRAGMA_ENABLE_DEPRECATION_WARNINGS
 				}
+				PRAGMA_ENABLE_DEPRECATION_WARNINGS
 			}));
 		}		
 
 		Menu->AddDynamicSection("AddMenuOptions", FNewToolMenuDelegate::CreateLambda([](UToolMenu* InMenu)
 		{
-			UContentBrowserAssetContextMenuContext* Context = InMenu->FindContext<UContentBrowserAssetContextMenuContext>();
+			const UContentBrowserAssetContextMenuContext* Context = InMenu->FindContext<UContentBrowserAssetContextMenuContext>();
 			if (Context && Context->AssetContextMenu.IsValid())
 			{
 				Context->AssetContextMenu.Pin()->AddMenuOptions(InMenu);
@@ -621,9 +634,9 @@ bool FAssetContextMenu::AddAssetTypeMenuOptions(UToolMenu* Menu)
 	{
 		// Label "GetAssetActions" section
 		FToolMenuSection& Section = Menu->FindOrAddSection("GetAssetActions");
-		if (Context->CommonAssetTypeActions.IsValid())
+		if (Context->CommonAssetDefinition)
 		{
-			Section.Label = FText::Format(NSLOCTEXT("AssetTools", "AssetSpecificOptionsMenuHeading", "{0} Actions"), Context->CommonAssetTypeActions.Pin()->GetName());
+			Section.Label = FText::Format(NSLOCTEXT("AssetTools", "AssetSpecificOptionsMenuHeading", "{0} Actions"), Context->CommonAssetDefinition->GetAssetDisplayName());
 		}
 		else if (Context->CommonClass)
 		{
diff --git a/Engine/Source/Editor/ContentBrowser/Private/ContentBrowserMenuContexts.cpp b/Engine/Source/Editor/ContentBrowser/Private/ContentBrowserMenuContexts.cpp
--- a/Engine/Source/Editor/ContentBrowser/Private/ContentBrowserMenuContexts.cpp
+++ b/Engine/Source/Editor/ContentBrowser/Private/ContentBrowserMenuContexts.cpp
@@ -86,11 +86,14 @@ namespace UE::ContentBrowser
 {
 	UToolMenu* ExtendToolMenu_AssetContextMenu(UClass* AssetClass)
 	{
-		check(AssetClass);
-
+		return ExtendToolMenu_AssetContextMenu(TSoftClassPtr<UObject>(AssetClass));
+	}
+	
+	UToolMenu* ExtendToolMenu_AssetContextMenu(TSoftClassPtr<UObject> AssetSoftClass)
+	{
 		FNameBuilder Builder;
 		Builder << TEXT("ContentBrowser.AssetContextMenu.");
-		AssetClass->GetFName().AppendString(Builder);
+		Builder << AssetSoftClass.GetAssetName();
 		const FName MenuName(Builder.ToView());
 
 		return UToolMenus::Get()->ExtendMenu(MenuName);
diff --git a/Engine/Source/Editor/ContentBrowser/Private/SThumbnailEditModeTools.cpp b/Engine/Source/Editor/ContentBrowser/Private/SThumbnailEditModeTools.cpp
--- a/Engine/Source/Editor/ContentBrowser/Private/SThumbnailEditModeTools.cpp
+++ b/Engine/Source/Editor/ContentBrowser/Private/SThumbnailEditModeTools.cpp
@@ -127,7 +127,7 @@ const FSlateBrush* SThumbnailEditModeTools::GetCurrentPrimitiveBrush() const
 	if ( ThumbnailInfo )
 	{
 		// Note this is for the icon only.  we are assuming the thumbnail renderer does the right thing when rendering
-		EThumbnailPrimType PrimType = ThumbnailInfo->bUserModifiedShape ? ThumbnailInfo->PrimitiveType.GetValue() : GetDefaultThumbnailType();
+		EThumbnailPrimType PrimType = ThumbnailInfo->bUserModifiedShape ? ThumbnailInfo->PrimitiveType.GetValue() : (EThumbnailPrimType)ThumbnailInfo->DefaultPrimitiveType.Get(EThumbnailPrimType::TPT_Sphere);
 		switch (PrimType)
 		{
 		case TPT_None: return FAppStyle::GetBrush("ContentBrowser.PrimitiveCustom");
@@ -354,30 +354,6 @@ USceneThumbnailInfoWithPrimitive* SThumbnailEditModeTools::ConstGetSceneThumbnai
 	return Cast<USceneThumbnailInfoWithPrimitive>( SceneThumbnailInfo.Get() );
 }
 
-EThumbnailPrimType SThumbnailEditModeTools::GetDefaultThumbnailType() const 
-{
-	EThumbnailPrimType DefaultPrimitiveType = TPT_Sphere;
-
-	if (AssetThumbnail.IsValid())
-	{
-		UObject* Asset = AssetThumbnail.Pin()->GetAsset();
-
-		if (Asset)
-		{
-			static const FName AssetToolsName("AssetTools");
-
-			FAssetToolsModule& AssetToolsModule = FModuleManager::LoadModuleChecked<FAssetToolsModule>(AssetToolsName);
-			TWeakPtr<IAssetTypeActions> AssetTypeActions = AssetToolsModule.Get().GetAssetTypeActionsForClass(Asset->GetClass());
-			if (AssetTypeActions.IsValid())
-			{
-				DefaultPrimitiveType = AssetTypeActions.Pin()->GetDefaultThumbnailPrimitiveType(Asset);
-			}
-		}
-	}
-
-	return DefaultPrimitiveType;
-}
-
 void SThumbnailEditModeTools::OnAssetDataChanged()
 {
 	// Set the SceneThumbnailInfo pointer as needed
diff --git a/Engine/Source/Editor/ContentBrowser/Public/ContentBrowserMenuContexts.h b/Engine/Source/Editor/ContentBrowser/Public/ContentBrowserMenuContexts.h
--- a/Engine/Source/Editor/ContentBrowser/Public/ContentBrowserMenuContexts.h
+++ b/Engine/Source/Editor/ContentBrowser/Public/ContentBrowserMenuContexts.h
@@ -6,6 +6,7 @@
 #include "Containers/Array.h"
 #include "ContentBrowserDelegates.h"
 #include "CoreMinimal.h"
+#include "Algo/SelectRandomWeighted.h"
 #include "HAL/Platform.h"
 #include "Templates/SharedPointer.h"
 #include "UObject/NameTypes.h"
@@ -21,6 +22,7 @@ class SAssetView;
 class SContentBrowser;
 class SFilterList;
 class UClass;
+class UAssetDefinition;
 struct FFrame;
 struct FToolMenuSection;
 struct FToolMenuContext;
@@ -40,9 +42,13 @@ public:
 
 	TWeakPtr<FAssetContextMenu> AssetContextMenu;
 
+	UE_DEPRECATED(5.2, "IAssetTypeActions (CommonAssetTypeActions) is being phased out too much menu related stuff was built into their design that forced loading assets.  Please use UAssetDefinition (CommonAssetDefinition).")
 	TWeakPtr<IAssetTypeActions> CommonAssetTypeActions;
 	
-	UE_DEPRECATED(5.1, "Use SelectedAssets now, this field will not contain any objects.")
+	UPROPERTY()
+	TObjectPtr<const UAssetDefinition> CommonAssetDefinition;
+	
+	UE_DEPRECATED(5.1, "Use SelectedAssets now, this field will not contain any objects.  You should call LoadSelectedObjects() based on what you need.")
 	TArray<TWeakObjectPtr<UObject>> SelectedObjects;
 
 	/**
@@ -57,7 +63,7 @@ public:
 	UPROPERTY()
 	bool bCanBeModified;
 
-	//UE_DEPRECATED(5.1, "GetSelectedObjects has been deprecated.  We no longer automatically load assets on right click.  Please use SelectedAssets and determine whatever you need for your context menu options without actually loading the assets.  When you finally need all or a subset of the selected assets use LoadSelectedAssets or LoadSelectedAssetsIf")
+	//UE_DEPRECATED(5.2, "GetSelectedObjects has been deprecated.  We no longer automatically load assets on right click.  Please use SelectedAssets and determine whatever you need for your context menu options without actually loading the assets.  When you finally need all or a subset of the selected assets use LoadSelectedAssets or LoadSelectedAssetsIf")
 	UFUNCTION(BlueprintCallable, Category="Tool Menus", meta=(DeprecatedFunction, DeprecationMessage = "GetSelectedObjects has been deprecated.  We no longer automatically load assets on right click.  If you can work without loading the assets, please use SelectedAssets.  Otherwise call LoadSelectedObjects"))
 	TArray<UObject*> GetSelectedObjects() const
 	{
@@ -239,4 +245,5 @@ class UToolMenu;
 namespace UE::ContentBrowser
 {
 	CONTENTBROWSER_API UToolMenu* ExtendToolMenu_AssetContextMenu(UClass* AssetClass);
+	CONTENTBROWSER_API UToolMenu* ExtendToolMenu_AssetContextMenu(TSoftClassPtr<UObject> AssetSoftClass);
 }
\ No newline at end of file
diff --git a/Engine/Source/Editor/UnrealEd/Classes/ThumbnailRendering/SceneThumbnailInfoWithPrimitive.h b/Engine/Source/Editor/UnrealEd/Classes/ThumbnailRendering/SceneThumbnailInfoWithPrimitive.h
--- a/Engine/Source/Editor/UnrealEd/Classes/ThumbnailRendering/SceneThumbnailInfoWithPrimitive.h
+++ b/Engine/Source/Editor/UnrealEd/Classes/ThumbnailRendering/SceneThumbnailInfoWithPrimitive.h
@@ -23,6 +23,8 @@ class USceneThumbnailInfoWithPrimitive : public USceneThumbnailInfo
 	/** The type of primitive used in this thumbnail */
 	UPROPERTY(EditAnywhere, Category=Thumbnail)
 	TEnumAsByte<EThumbnailPrimType> PrimitiveType;
+	
+	TOptional<TEnumAsByte<EThumbnailPrimType>> DefaultPrimitiveType;
 
 	/** The custom mesh used when the primitive type is TPT_None */
 	UPROPERTY(EditAnywhere, Category=Thumbnail, meta = (AllowedClasses = "/Script/Engine.StaticMesh,/Script/Engine.SkeletalMesh", ExactClass = "true"))
@@ -30,6 +32,7 @@ class USceneThumbnailInfoWithPrimitive : public USceneThumbnailInfo
 
 	UPROPERTY()
 	bool bUserModifiedShape;
+	
 public:
 	virtual void Serialize(FArchive& Ar) override;
 	virtual void PostLoad() override;
diff --git a/Engine/Source/Editor/UnrealEd/Private/AssetDefinitionDefault.cpp b/Engine/Source/Editor/UnrealEd/Private/AssetDefinitionDefault.cpp
new file mode 100644
--- /dev/null
+++ b/Engine/Source/Editor/UnrealEd/Private/AssetDefinitionDefault.cpp
@@ -0,0 +1,80 @@
+// Copyright Epic Games, Inc. All Rights Reserved.
+
+#include "AssetDefinitionDefault.h"
+
+#include "AssetToolsModule.h"
+#include "EditorFramework/AssetImportData.h"
+#include "Settings/EditorLoadingSavingSettings.h"
+#include "Toolkits/SimpleAssetEditor.h"
+
+#define LOCTEXT_NAMESPACE "AssetDefinitionDefault"
+
+EAssetCommandResult UAssetDefinitionDefault::OpenAssets(const FAssetOpenArgs& OpenArgs) const
+{
+	if (OpenArgs.OpenMethod == EAssetOpenMethod::Edit)
+	{
+		FSimpleAssetEditor::CreateEditor(EToolkitMode::Standalone, OpenArgs.ToolkitHost, OpenArgs.LoadObjects<UObject>());
+		return EAssetCommandResult::Handled;
+	}
+
+	return EAssetCommandResult::Unhandled;
+}
+
+EAssetCommandResult UAssetDefinitionDefault::GetSourceFiles(const FAssetSourceFileArgs& SourceFileArgs, TArray<FAssetSourceFile>& OutSourceAssets) const
+{
+	for (const FAssetData& Asset : SourceFileArgs.Assets)
+    {
+    	FString SourceFileTagData;
+    	if (Asset.GetTagValue(UObject::SourceFileTagName(), SourceFileTagData))
+    	{
+    		TOptional<FAssetImportInfo> ImportInfo = FAssetImportInfo::FromJson(SourceFileTagData);
+    		if (ImportInfo.IsSet())
+    		{
+    			for (const FAssetImportInfo::FSourceFile& SourceFile : ImportInfo->SourceFiles)
+    			{
+    				FAssetSourceFile AssetSourceFile;
+    				AssetSourceFile.DisplayLabelName = SourceFile.DisplayLabelName;
+    				AssetSourceFile.RelativeFilename = SourceFile.RelativeFilename;
+    				OutSourceAssets.Add(MoveTemp(AssetSourceFile));
+    			}
+    		}
+    	}
+    }
+    
+    return EAssetCommandResult::Handled;
+}
+
+EAssetCommandResult UAssetDefinitionDefault::PerformAssetDiff(const FAssetDiffArgs& DiffArgs) const
+{
+	check(DiffArgs.OldAsset != nullptr);
+	check(DiffArgs.NewAsset != nullptr);
+
+	IAssetTools& AssetTools = FModuleManager::Get().LoadModuleChecked<FAssetToolsModule>("AssetTools").Get();
+
+	// Dump assets to temp text files
+	FString OldTextFilename = AssetTools.DumpAssetToTempFile(DiffArgs.OldAsset);
+	FString NewTextFilename = AssetTools.DumpAssetToTempFile(DiffArgs.NewAsset);
+	FString DiffCommand = GetDefault<UEditorLoadingSavingSettings>()->TextDiffToolPath.FilePath;
+
+	AssetTools.CreateDiffProcess(DiffCommand, OldTextFilename, NewTextFilename);
+
+	return EAssetCommandResult::Handled;
+}
+
+/*
+virtual void PerformAssetDiff(UObject* OldAsset, UObject* NewAsset, const struct FRevisionInfo& OldRevision, const struct FRevisionInfo& NewRevision) const override
+{
+	check(OldAsset != nullptr);
+	check(NewAsset != nullptr);
+
+	// Dump assets to temp text files
+	FString OldTextFilename = DumpAssetToTempFile(OldAsset);
+	FString NewTextFilename = DumpAssetToTempFile(NewAsset);
+	FString DiffCommand = GetDefault<UEditorLoadingSavingSettings>()->TextDiffToolPath.FilePath;
+
+	FAssetToolsModule& AssetToolsModule = FModuleManager::Get().LoadModuleChecked<FAssetToolsModule>("AssetTools");
+	AssetToolsModule.Get().CreateDiffProcess(DiffCommand, OldTextFilename, NewTextFilename);
+}
+*/
+
+#undef LOCTEXT_NAMESPACE
diff --git a/Engine/Source/Editor/UnrealEd/Private/Editor.cpp b/Engine/Source/Editor/UnrealEd/Private/Editor.cpp
--- a/Engine/Source/Editor/UnrealEd/Private/Editor.cpp
+++ b/Engine/Source/Editor/UnrealEd/Private/Editor.cpp
@@ -61,6 +61,8 @@
 
 // AIMdule
 
+#include "AssetDefinition.h"
+#include "AssetDefinitionRegistry.h"
 #include "Framework/Notifications/NotificationManager.h"
 #include "Widgets/Notifications/SNotificationList.h"
 #include "K2Node_AddComponent.h"
@@ -736,15 +738,6 @@ void FReimportManager::GetNewReimportPath(UObject* Obj, TArray<FString>& InOutFi
 	FString FileTypes;
 	FString AllExtensions;
 	TArray<UFactory*> Factories;
-	TArray<FString> SourceFileLabels;
-	FAssetToolsModule& AssetToolsModule = FModuleManager::Get().LoadModuleChecked<FAssetToolsModule>("AssetTools");
-	const auto AssetTypeActions = AssetToolsModule.Get().GetAssetTypeActionsForClass(Obj->GetClass());
-	if (AssetTypeActions.IsValid())
-	{
-		TArray<UObject*> Objects;
-		Objects.Add(Obj);
-		AssetTypeActions.Pin()->GetSourceFileLabels(Objects, SourceFileLabels);
-	}
 
 	// Determine whether we will allow multi select and clear old filenames
 	bool bAllowMultiSelect = SourceFileIndex == INDEX_NONE && InOutFilenames.Num() > 1;
@@ -842,11 +835,18 @@ void FReimportManager::GetNewReimportPath(UObject* Obj, TArray<FString>& InOutFi
 		FString Title = FString::Printf(TEXT("%s: %s"), *NSLOCTEXT("ReimportManager", "ImportDialogTitle", "Import For").ToString(), *Obj->GetName());
 		if (SourceFileIndex != INDEX_NONE)
 		{
-			if (SourceFileLabels.IsValidIndex(SourceFileIndex))
+			FAssetData AssetData(Obj);
+			const UAssetDefinition* AssetDefinition = UAssetDefinitionRegistry::Get()->GetAssetDefinitionForAsset(AssetData);
+
+			TArray<FAssetSourceFile> OutSourceAssets;
+			FAssetSourceFileArgs SourceFileArgs({ AssetData });
+			AssetDefinition->GetSourceFiles(SourceFileArgs, OutSourceAssets);
+			
+			if (OutSourceAssets.IsValidIndex(SourceFileIndex))
 			{
 				Title = FString::Printf(TEXT("%s %s %s: %s"),
 					*NSLOCTEXT("ReimportManager", "ImportDialogTitleLabelPart1", "Select").ToString(),
-					*SourceFileLabels[SourceFileIndex],
+					*OutSourceAssets[SourceFileIndex].DisplayLabelName,
 					*NSLOCTEXT("ReimportManager", "ImportDialogTitleLabelPart2", "Source File For").ToString(),
 					*Obj->GetName());
 			}
diff --git a/Engine/Source/Editor/UnrealEd/Public/AssetDefinitionDefault.h b/Engine/Source/Editor/UnrealEd/Public/AssetDefinitionDefault.h
new file mode 100644
--- /dev/null
+++ b/Engine/Source/Editor/UnrealEd/Public/AssetDefinitionDefault.h
@@ -0,0 +1,26 @@
+// Copyright Epic Games, Inc. All Rights Reserved.
+
+#pragma once
+
+#include "CoreMinimal.h"
+#include "Engine/World.h"
+#include "AssetDefinition.h"
+
+#include "AssetDefinitionDefault.generated.h"
+
+
+UCLASS(Abstract)
+class UNREALED_API UAssetDefinitionDefault : public UAssetDefinition
+{
+	GENERATED_BODY()
+
+public:
+	// UAssetDefinition Begin
+	virtual EAssetCommandResult OpenAssets(const FAssetOpenArgs& OpenArgs) const override;
+	virtual EAssetCommandResult GetSourceFiles(const FAssetSourceFileArgs& SourceFileArgs, TArray<FAssetSourceFile>& OutSourceAssets) const override;
+	virtual EAssetCommandResult PerformAssetDiff(const FAssetDiffArgs& DiffArgs) const override;
+	// UAssetDefinition End
+	
+public:
+	
+};
diff --git a/Engine/Source/Editor/UnrealEd/UnrealEd.Build.cs b/Engine/Source/Editor/UnrealEd/UnrealEd.Build.cs
--- a/Engine/Source/Editor/UnrealEd/UnrealEd.Build.cs
+++ b/Engine/Source/Editor/UnrealEd/UnrealEd.Build.cs
@@ -79,6 +79,7 @@ public class UnrealEd : ModuleRules
 		PublicDependencyModuleNames.AddRange(
 			new string[]
 			{
+				"AssetDefinition",
 				"Core",
 				"CoreUObject",
 				"ApplicationCore",
diff --git a/Engine/Source/Runtime/AssetRegistry/Private/AssetRegistryHelpers.cpp b/Engine/Source/Runtime/AssetRegistry/Private/AssetRegistryHelpers.cpp
--- a/Engine/Source/Runtime/AssetRegistry/Private/AssetRegistryHelpers.cpp
+++ b/Engine/Source/Runtime/AssetRegistry/Private/AssetRegistryHelpers.cpp
@@ -87,6 +87,32 @@ FARFilter UAssetRegistryHelpers::SetFilterTagsAndValues(const FARFilter& InFilte
 	return FilterCopy;
 }
 
+UClass* UAssetRegistryHelpers::FindAssetNativeClass(const FAssetData& AssetData)
+{
+	UClass* AssetClass = AssetData.GetClass();
+	if (AssetClass == nullptr)
+	{
+		const IAssetRegistry& AssetRegistry = IAssetRegistry::GetChecked();
+		
+		TArray<FTopLevelAssetPath> AncestorClasses;
+		AssetRegistry.GetAncestorClassNames(AssetData.AssetClassPath, AncestorClasses);
+		for (const FTopLevelAssetPath& AncestorClassPath : AncestorClasses)
+		{
+			AssetClass = FindObject<UClass>(AncestorClassPath);
+			if (AssetClass)
+			{
+				break;
+			}
+		}
+	}
+	while (AssetClass && !AssetClass->HasAnyClassFlags(CLASS_Native))
+	{
+		AssetClass = AssetClass->GetSuperClass();
+	}
+
+	return AssetClass;
+}
+
 void UAssetRegistryHelpers::GetBlueprintAssets(const FARFilter& InFilter, TArray<FAssetData>& OutAssetData)
 {
 	IAssetRegistry& AssetRegistry = IAssetRegistry::GetChecked();
diff --git a/Engine/Source/Runtime/AssetRegistry/Public/AssetRegistry/AssetRegistryHelpers.h b/Engine/Source/Runtime/AssetRegistry/Public/AssetRegistry/AssetRegistryHelpers.h
--- a/Engine/Source/Runtime/AssetRegistry/Public/AssetRegistry/AssetRegistryHelpers.h
+++ b/Engine/Source/Runtime/AssetRegistry/Public/AssetRegistry/AssetRegistryHelpers.h
@@ -101,6 +101,14 @@ public:
 	UFUNCTION(BlueprintPure, Category = "Asset Registry", meta=(ScriptMethod))
 	static void GetBlueprintAssets(const FARFilter& InFilter, TArray<FAssetData>& OutAssetData);
 
+	/**
+	 * Returns the first native class of the asset type that can be found.  Normally this is just the FAssetData::GetClass(),
+	 * however if the class is a blueprint generated class it may not be loaded.  In which case GetAncestorClassNames will
+	 * be used to find the first native super class.  This can be slow if temporary caching mode is not on.
+	 */
+	UFUNCTION(BlueprintPure, Category = "Asset Registry", meta=(ScriptMethod))
+	static ASSETREGISTRY_API UClass* FindAssetNativeClass(const FAssetData& AssetData);
+	
 	/** Enable/disable asset registry caching mode for the duration of the scope */
 	struct ASSETREGISTRY_API FTemporaryCachingModeScope
 	{
