AssetTypeActions - Fixing a case where the "Basic" category needs special rules since it's not actually in the Advanced set.  The new system doesn't do this, it just has categories.

diff --git a/Engine/Plugins/Editor/EngineAssetDefinitions/Source/Private/AssetDefinition_World.cpp b/Engine/Plugins/Editor/EngineAssetDefinitions/Source/Private/AssetDefinition_World.cpp
--- a/Engine/Plugins/Editor/EngineAssetDefinitions/Source/Private/AssetDefinition_World.cpp
+++ b/Engine/Plugins/Editor/EngineAssetDefinitions/Source/Private/AssetDefinition_World.cpp
@@ -10,6 +10,12 @@
 
 #define LOCTEXT_NAMESPACE "UAssetDefinition_World"
 
+TConstArrayView<FAssetCategoryPath> UAssetDefinition_World::GetAssetCategories() const
+{
+	static const auto Categories = { EAssetCategoryPaths::Basic };
+	return Categories;
+}
+
 TArray<FAssetData> UAssetDefinition_World::PrepareToActivateAssets(const FAssetActivateArgs& ActivateArgs) const
 {
 	TArray<FAssetData> AssetsToOpen;
diff --git a/Engine/Plugins/Editor/EngineAssetDefinitions/Source/Public/AssetDefinition_World.h b/Engine/Plugins/Editor/EngineAssetDefinitions/Source/Public/AssetDefinition_World.h
--- a/Engine/Plugins/Editor/EngineAssetDefinitions/Source/Public/AssetDefinition_World.h
+++ b/Engine/Plugins/Editor/EngineAssetDefinitions/Source/Public/AssetDefinition_World.h
@@ -18,11 +18,7 @@ public:
 	virtual FText GetAssetDisplayName() const override { return NSLOCTEXT("AssetTypeActions", "AssetTypeActions_World", "Level"); }
 	virtual FLinearColor GetAssetColor() const override { return FAppStyle::Get().GetColor("LevelEditor.AssetColor"); }
 	virtual TSoftClassPtr<UObject> GetAssetClass() const override { return UWorld::StaticClass(); }
-	virtual TConstArrayView<FAssetCategoryPath> GetAssetCategories() const override
-	{
-		static const auto Categories = { EAssetCategoryPaths::Basic };
-		return Categories;
-	}
+	virtual TConstArrayView<FAssetCategoryPath> GetAssetCategories() const override;
 
 	virtual TArray<FAssetData> PrepareToActivateAssets(const FAssetActivateArgs& ActivateArgs) const override;
 	virtual EAssetCommandResult OpenAssets(const FAssetOpenArgs& OpenArgs) const override;
diff --git a/Engine/Source/Developer/AssetTools/Private/AssetDefinition_AssetTypeActionsProxy.cpp b/Engine/Source/Developer/AssetTools/Private/AssetDefinition_AssetTypeActionsProxy.cpp
--- a/Engine/Source/Developer/AssetTools/Private/AssetDefinition_AssetTypeActionsProxy.cpp
+++ b/Engine/Source/Developer/AssetTools/Private/AssetDefinition_AssetTypeActionsProxy.cpp
@@ -5,6 +5,8 @@
 #include "IAssetTools.h"
 #include "Misc/AssetFilterData.h"
 
+#define LOCTEXT_NAMESPACE "UAssetDefinition_AssetTypeActionsProxy"
+
 void UAssetDefinition_AssetTypeActionsProxy::Initialize(const TSharedRef<IAssetTypeActions>& NewActions)
 {
 	AssetType = NewActions;
@@ -40,6 +42,12 @@ TConstArrayView<FAssetCategoryPath> UAssetDefinition_AssetTypeActionsProxy::GetA
     AssetTools.GetAllAdvancedAssetCategories(AdvancedAssetCategories);
 	
 	const uint32 CategoryBits = AssetType->GetCategories();
+
+	if (CategoryBits & EAssetTypeCategories::Basic)
+	{
+		AssetCategories.Add(FAssetCategoryPath({ EAssetCategoryPaths::Basic }));
+	}
+	
 	for (const FAdvancedAssetCategory& Category : AdvancedAssetCategories)
 	{
 		if (Category.CategoryType & CategoryBits)
@@ -226,4 +234,6 @@ const FSlateBrush* UAssetDefinition_AssetTypeActionsProxy::GetIconBrush(const FA
 TSharedPtr<SWidget> UAssetDefinition_AssetTypeActionsProxy::GetThumbnailOverlay(const FAssetData& InAssetData) const
 {
 	return AssetType->GetThumbnailOverlay(InAssetData);
-}
\ No newline at end of file
+}
+
+#undef LOCTEXT_NAMESPACE
\ No newline at end of file
diff --git a/Engine/Source/Developer/AssetTools/Private/AssetTools.cpp b/Engine/Source/Developer/AssetTools/Private/AssetTools.cpp
--- a/Engine/Source/Developer/AssetTools/Private/AssetTools.cpp
+++ b/Engine/Source/Developer/AssetTools/Private/AssetTools.cpp
@@ -225,10 +225,17 @@ public:
 		uint32 CategoryBits = 0;
 		for (const FAssetCategoryPath& Category : AssetDefinitionPtr.Get()->GetAssetCategories())
 		{
-			CategoryBits |= AssetTools.FindAdvancedAssetCategory(Category.GetCategory());
+			if (Category.GetCategory() == EAssetCategoryPaths::Basic.GetCategory())
+			{
+				CategoryBits |= EAssetTypeCategories::Basic;
+			}
+			else
+			{
+				CategoryBits |= AssetTools.FindAdvancedAssetCategory(Category.GetCategory());	
+			}
 		}
 
-		return EAssetTypeCategories::Misc;
+		return CategoryBits;
 	}
 
 private:
