AssetDefinitionProxy - Submenus now get initialized.

diff --git a/Engine/Source/Developer/AssetTools/Private/AssetTools.cpp b/Engine/Source/Developer/AssetTools/Private/AssetTools.cpp
--- a/Engine/Source/Developer/AssetTools/Private/AssetTools.cpp
+++ b/Engine/Source/Developer/AssetTools/Private/AssetTools.cpp
@@ -240,7 +240,7 @@ public:
 
 private:
 	TArray<FText> SubMenus;
-	bool SubmenusInitialized = false;
+	mutable bool SubmenusInitialized = false;
 
 public:
 	/** Returns array of sub-menu names that this asset type is parented under in the Asset Creation Context Menu. */
@@ -250,6 +250,8 @@ public:
 
 		if (!SubmenusInitialized)
 		{
+			SubmenusInitialized = true;
+			
 			for (const FAssetCategoryPath& Category : AssetDefinitionPtr.Get()->GetAssetCategories())
 			{
 				if (Category.HasSubCategory())
