Fix Advanced Copy not fixing up external actors reference to packages that are part of the copy.

diff --git a/Engine/Source/Developer/AssetTools/Private/AssetTools.cpp b/Engine/Source/Developer/AssetTools/Private/AssetTools.cpp
--- a/Engine/Source/Developer/AssetTools/Private/AssetTools.cpp
+++ b/Engine/Source/Developer/AssetTools/Private/AssetTools.cpp
@@ -1990,6 +1990,7 @@ bool UAssetToolsImpl::AdvancedCopyPackages(
 
 		TSet<UObject*> ExistingObjectSet;
 		TSet<UObject*> NewObjectSet;
+		TSet<FName> CopiedWorldPartitionMaps;
 		FString CopyErrors;
 
 		SuccessfullyCopiedDestinationFiles.Reserve(SourceAndDestPackages.Num());
@@ -2025,6 +2026,8 @@ bool UAssetToolsImpl::AdvancedCopyPackages(
 						PGN.PackageName = DestFilename;
 						const bool bShouldPromptForDestinationConflict = !bCopyOverAllDestinationOverlaps;
 						TMap<TSoftObjectPtr<UObject>, TSoftObjectPtr<UObject>> DuplicatedObjects;
+						FName PackageFName(*PackageName);
+
 
 						// Temp fix for some codepaths that allows advanced copy of world packages. For partitioned worlds, this can only be supported for worlds with
 						// streaming disabled and this code should be removed once the callers switch to the same codepath as editor save as.
@@ -2037,6 +2040,7 @@ bool UAssetToolsImpl::AdvancedCopyPackages(
 								{
 									WorldPartition->Initialize(World, FTransform::Identity);
 								}
+								CopiedWorldPartitionMaps.Add(PackageFName);
 							}
 						}
 
@@ -2045,7 +2049,7 @@ bool UAssetToolsImpl::AdvancedCopyPackages(
 							ExistingObjectSet.Add(ExistingObject);
 							NewObjectSet.Add(NewObject);
 							DuplicatedObjectsForEachPackage.Add(MoveTemp(DuplicatedObjects));
-							SuccessfullyCopiedSourcePackages.Add(FName(*PackageName));
+							SuccessfullyCopiedSourcePackages.Add(PackageFName);
 							SuccessfullyCopiedDestinationFiles.Add(DestFilename);
 							SuccessfullyCopiedDestinationPackages.Add(NewObject->GetPackage());
 						}
@@ -2064,6 +2068,21 @@ bool UAssetToolsImpl::AdvancedCopyPackages(
 		{
 			Dependencies.Reset();
 			AssetRegistryModule.Get().GetDependencies(SuccessfullyCopiedPackage, Dependencies);
+
+			// Temp fix for some codepaths that allows advanced copy of world packages.
+			// if the map is a world partition map, add dependencies of the actor packages as well
+			if (CopiedWorldPartitionMaps.Contains(SuccessfullyCopiedPackage))
+			{
+				TArray<FName> ExternalObjectsPaths;
+				Algo::Transform(ULevel::GetExternalObjectsPaths(SuccessfullyCopiedPackage.ToString()), ExternalObjectsPaths, [](const FString& Path) { return FName(*Path); });
+				TArray<FAssetData> ExternalActors;
+				AssetRegistryModule.Get().GetAssetsByPaths(ExternalObjectsPaths, ExternalActors, true, true);
+				for (const FAssetData& AssetData : ExternalActors)
+				{
+					AssetRegistryModule.Get().GetDependencies(AssetData.PackageName, Dependencies);
+				}
+			}
+
 			for (FName Dependency : Dependencies)
 			{
 				const int32 DependencyIndex = SuccessfullyCopiedSourcePackages.IndexOfByKey(Dependency);
