AssetDefinition - Converting another batch of AssetTypeActions to AssetDefinitions.

diff --git a/Engine/Plugins/Editor/EngineAssetDefinitions/Source/EngineAssetDefinitions.Build.cs b/Engine/Plugins/Editor/EngineAssetDefinitions/Source/EngineAssetDefinitions.Build.cs
--- a/Engine/Plugins/Editor/EngineAssetDefinitions/Source/EngineAssetDefinitions.Build.cs
+++ b/Engine/Plugins/Editor/EngineAssetDefinitions/Source/EngineAssetDefinitions.Build.cs
@@ -30,6 +30,8 @@ public class EngineAssetDefinitions : ModuleRules
 				"Slate",
 				"Kismet",
 				"DesktopPlatform",
+				"Foliage",
+				"PhysicsCore",
 			}
 		);
 		
@@ -40,6 +42,7 @@ public class EngineAssetDefinitions : ModuleRules
 				"StaticMeshEditor",
 				"TextureEditor",
 				"MaterialEditor",
+				"PhysicsAssetEditor",
 			}
 		);
 
@@ -50,6 +53,7 @@ public class EngineAssetDefinitions : ModuleRules
 				"StaticMeshEditor",
 				"TextureEditor",
 				"MaterialEditor",
+				"PhysicsAssetEditor",
 			}
 		);
 	}
diff --git a/Engine/Plugins/Editor/EngineAssetDefinitions/Source/Private/AssetDefinition_MaterialParameterCollection.h b/Engine/Plugins/Editor/EngineAssetDefinitions/Source/Private/AssetDefinition_MaterialParameterCollection.h
new file mode 100644
--- /dev/null
+++ b/Engine/Plugins/Editor/EngineAssetDefinitions/Source/Private/AssetDefinition_MaterialParameterCollection.h
@@ -0,0 +1,27 @@
+// Copyright Epic Games, Inc. All Rights Reserved.
+
+#pragma once
+
+#include "CoreMinimal.h"
+#include "Materials/MaterialParameterCollection.h"
+#include "AssetDefinitionDefault.h"
+
+#include "AssetDefinition_MaterialParameterCollection.generated.h"
+
+UCLASS()
+class UAssetDefinition_MaterialParameterCollection : public UAssetDefinitionDefault
+{
+	GENERATED_BODY()
+
+public:
+	// UAssetDefinition Begin
+	virtual FText GetAssetDisplayName() const override { return NSLOCTEXT("AssetTypeActions", "AssetTypeActions_MaterialParameterCollection", "Material Parameter Collection"); }
+	virtual FLinearColor GetAssetColor() const override { return FLinearColor(FColor(0, 192, 0)); }
+	virtual TSoftClassPtr<UObject> GetAssetClass() const override { return UMaterialParameterCollection::StaticClass(); }
+	virtual TConstArrayView<FAssetCategoryPath> GetAssetCategories() const override
+	{
+		static const auto Categories = { EAssetCategoryPaths::Material };
+		return Categories;
+	}
+	// UAssetDefinition End
+};
diff --git a/Engine/Plugins/Editor/EngineAssetDefinitions/Source/Private/AssetDefinition_NodeMappingContainer.h b/Engine/Plugins/Editor/EngineAssetDefinitions/Source/Private/AssetDefinition_NodeMappingContainer.h
new file mode 100644
--- /dev/null
+++ b/Engine/Plugins/Editor/EngineAssetDefinitions/Source/Private/AssetDefinition_NodeMappingContainer.h
@@ -0,0 +1,27 @@
+// Copyright Epic Games, Inc. All Rights Reserved.
+
+#pragma once
+
+#include "CoreMinimal.h"
+#include "Animation/NodeMappingContainer.h"
+#include "AssetDefinitionDefault.h"
+
+#include "AssetDefinition_NodeMappingContainer.generated.h"
+
+UCLASS()
+class UAssetDefinition_NodeMappingContainer : public UAssetDefinitionDefault
+{
+	GENERATED_BODY()
+
+public:
+	// UAssetDefinition Begin
+	virtual FText GetAssetDisplayName() const override { return NSLOCTEXT("AssetTypeActions", "AssetTypeActions_NodeMappingContainer", "Node Mapping Container"); }
+	virtual FLinearColor GetAssetColor() const override { return FLinearColor(FColor(112,146,190)); }
+	virtual TSoftClassPtr<UObject> GetAssetClass() const override { return UNodeMappingContainer::StaticClass(); }
+	virtual TConstArrayView<FAssetCategoryPath> GetAssetCategories() const override
+	{
+		static const auto Categories = { EAssetCategoryPaths::Animation };
+		return Categories;
+	}
+	// UAssetDefinition End
+};
diff --git a/Engine/Plugins/Editor/EngineAssetDefinitions/Source/Private/AssetDefinition_ObjectLibrary.h b/Engine/Plugins/Editor/EngineAssetDefinitions/Source/Private/AssetDefinition_ObjectLibrary.h
new file mode 100644
--- /dev/null
+++ b/Engine/Plugins/Editor/EngineAssetDefinitions/Source/Private/AssetDefinition_ObjectLibrary.h
@@ -0,0 +1,23 @@
+// Copyright Epic Games, Inc. All Rights Reserved.
+
+#pragma once
+
+#include "CoreMinimal.h"
+#include "AssetTypeActions_Base.h"
+#include "Engine/ObjectLibrary.h"
+#include "AssetDefinitionDefault.h"
+
+#include "AssetDefinition_ObjectLibrary.generated.h"
+
+UCLASS()
+class UAssetDefinition_ObjectLibrary : public UAssetDefinitionDefault
+{
+	GENERATED_BODY()
+
+public:
+	// UAssetDefinition Begin
+	virtual FText GetAssetDisplayName() const override { return NSLOCTEXT("AssetTypeActions", "AssetTypeActions_ObjectLibrary", "Object Library"); }
+	virtual FLinearColor GetAssetColor() const override { return FLinearColor(FColor(255, 63, 108)); }
+	virtual TSoftClassPtr<UObject> GetAssetClass() const override { return UObjectLibrary::StaticClass(); }
+	// UAssetDefinition End
+};
diff --git a/Engine/Plugins/Editor/EngineAssetDefinitions/Source/Private/AssetDefinition_PhysicalMaterial.h b/Engine/Plugins/Editor/EngineAssetDefinitions/Source/Private/AssetDefinition_PhysicalMaterial.h
new file mode 100644
--- /dev/null
+++ b/Engine/Plugins/Editor/EngineAssetDefinitions/Source/Private/AssetDefinition_PhysicalMaterial.h
@@ -0,0 +1,27 @@
+// Copyright Epic Games, Inc. All Rights Reserved.
+
+#pragma once
+
+#include "CoreMinimal.h"
+#include "PhysicalMaterials/PhysicalMaterial.h"
+#include "AssetDefinitionDefault.h"
+
+#include "AssetDefinition_PhysicalMaterial.generated.h"
+
+UCLASS()
+class UAssetDefinition_PhysicalMaterial : public UAssetDefinitionDefault
+{
+	GENERATED_BODY()
+
+public:
+	// UAssetDefinition Begin
+	virtual FText GetAssetDisplayName() const override { return NSLOCTEXT("AssetTypeActions", "AssetTypeActions_PhysicalMaterial", "Physical Material"); }
+	virtual FLinearColor GetAssetColor() const override { return FLinearColor(FColor(200,192,128)); }
+	virtual TSoftClassPtr<UObject> GetAssetClass() const override { return UPhysicalMaterial::StaticClass(); }
+	virtual TConstArrayView<FAssetCategoryPath> GetAssetCategories() const override
+	{
+		static const auto Categories = { EAssetCategoryPaths::Physics };
+		return Categories;
+	}
+	// UAssetDefinition End
+};
diff --git a/Engine/Plugins/Editor/EngineAssetDefinitions/Source/Private/AssetDefinition_PhysicsAsset.cpp b/Engine/Plugins/Editor/EngineAssetDefinitions/Source/Private/AssetDefinition_PhysicsAsset.cpp
new file mode 100644
--- /dev/null
+++ b/Engine/Plugins/Editor/EngineAssetDefinitions/Source/Private/AssetDefinition_PhysicsAsset.cpp
@@ -0,0 +1,34 @@
+// Copyright Epic Games, Inc. All Rights Reserved.
+
+#include "AssetDefinition_PhysicsAsset.h"
+#include "Modules/ModuleManager.h"
+#include "PhysicsAssetEditorModule.h"
+#include "ThumbnailRendering/SceneThumbnailInfo.h"
+
+#define LOCTEXT_NAMESPACE "AssetTypeActions"
+
+UThumbnailInfo* UAssetDefinition_PhysicsAsset::LoadThumbnailInfo(const FAssetData& InAsset) const
+{
+	UPhysicsAsset* PhysicsAsset = CastChecked<UPhysicsAsset>(InAsset.GetAsset());
+	UThumbnailInfo* ThumbnailInfo = PhysicsAsset->ThumbnailInfo;
+	if (ThumbnailInfo == NULL)
+	{
+		ThumbnailInfo = NewObject<USceneThumbnailInfo>(PhysicsAsset, NAME_None, RF_Transactional);
+		PhysicsAsset->ThumbnailInfo = ThumbnailInfo;
+	}
+
+	return ThumbnailInfo;
+}
+
+EAssetCommandResult UAssetDefinition_PhysicsAsset::OpenAssets(const FAssetOpenArgs& OpenArgs) const
+{
+	for (UPhysicsAsset* PhysicsAsset : OpenArgs.LoadObjects<UPhysicsAsset>())
+	{
+		IPhysicsAssetEditorModule* PhysicsAssetEditorModule = &FModuleManager::LoadModuleChecked<IPhysicsAssetEditorModule>( "PhysicsAssetEditor" );
+		PhysicsAssetEditorModule->CreatePhysicsAssetEditor(OpenArgs.GetToolkitMode(), OpenArgs.ToolkitHost, PhysicsAsset);
+	}
+
+	return EAssetCommandResult::Handled;
+}
+
+#undef LOCTEXT_NAMESPACE
diff --git a/Engine/Plugins/Editor/EngineAssetDefinitions/Source/Private/AssetDefinition_PhysicsAsset.h b/Engine/Plugins/Editor/EngineAssetDefinitions/Source/Private/AssetDefinition_PhysicsAsset.h
new file mode 100644
--- /dev/null
+++ b/Engine/Plugins/Editor/EngineAssetDefinitions/Source/Private/AssetDefinition_PhysicsAsset.h
@@ -0,0 +1,29 @@
+// Copyright Epic Games, Inc. All Rights Reserved.
+
+#pragma once
+
+#include "CoreMinimal.h"
+#include "PhysicsEngine/PhysicsAsset.h"
+#include "AssetDefinitionDefault.h"
+
+#include "AssetDefinition_PhysicsAsset.generated.h"
+
+UCLASS()
+class UAssetDefinition_PhysicsAsset : public UAssetDefinitionDefault
+{
+	GENERATED_BODY()
+
+public:
+	// UAssetDefinition Begin
+	virtual FText GetAssetDisplayName() const override { return NSLOCTEXT("AssetTypeActions", "AssetTypeActions_PhysicsAsset", "Physics Asset"); }
+	virtual FLinearColor GetAssetColor() const override { return FLinearColor(FColor(255,192,128)); }
+	virtual TSoftClassPtr<UObject> GetAssetClass() const override { return UPhysicsAsset::StaticClass(); }
+	virtual TConstArrayView<FAssetCategoryPath> GetAssetCategories() const override
+	{
+		static const auto Categories = { EAssetCategoryPaths::Physics };
+		return Categories;
+	}
+	virtual EAssetCommandResult OpenAssets(const FAssetOpenArgs& OpenArgs) const override;
+	virtual UThumbnailInfo* LoadThumbnailInfo(const FAssetData& InAsset) const override;
+	// UAssetDefinition End
+};
diff --git a/Engine/Plugins/Editor/EngineAssetDefinitions/Source/Private/AssetDefinition_PoseAsset.h b/Engine/Plugins/Editor/EngineAssetDefinitions/Source/Private/AssetDefinition_PoseAsset.h
new file mode 100644
--- /dev/null
+++ b/Engine/Plugins/Editor/EngineAssetDefinitions/Source/Private/AssetDefinition_PoseAsset.h
@@ -0,0 +1,17 @@
+// Copyright Epic Games, Inc. All Rights Reserved.
+
+#pragma once
+
+#include "CoreMinimal.h"
+#include "AssetTypeActions/AssetTypeActions_AnimationAsset.h"
+#include "Animation/PoseAsset.h"
+
+class FAssetTypeActions_PoseAsset : public FAssetTypeActions_AnimationAsset
+{
+public:
+	// IAssetTypeActions Implementation
+	virtual FText GetName() const override { return NSLOCTEXT("AssetTypeActions", "AssetTypeActions_PoseAsset", "Pose Asset"); }
+	virtual FColor GetTypeColor() const override { return FColor(237, 28, 36); }
+	virtual UClass* GetSupportedClass() const override { return UPoseAsset::StaticClass(); }
+	virtual bool CanFilter() override { return true; }
+};
diff --git a/Engine/Plugins/Editor/EngineAssetDefinitions/Source/Private/AssetDefinition_PreviewMeshCollection.h b/Engine/Plugins/Editor/EngineAssetDefinitions/Source/Private/AssetDefinition_PreviewMeshCollection.h
new file mode 100644
--- /dev/null
+++ b/Engine/Plugins/Editor/EngineAssetDefinitions/Source/Private/AssetDefinition_PreviewMeshCollection.h
@@ -0,0 +1,17 @@
+// Copyright Epic Games, Inc. All Rights Reserved.
+
+#pragma once
+
+#include "CoreMinimal.h"
+#include "AssetTypeActions_Base.h"
+#include "Engine/PreviewMeshCollection.h"
+
+class FAssetTypeActions_PreviewMeshCollection : public FAssetTypeActions_Base
+{
+public:
+	// IAssetTypeActions Implementation
+	virtual FText GetName() const override { return NSLOCTEXT("AssetTypeActions", "AssetTypeActions_PreviewMeshCollection", "Preview Mesh Collection"); }
+	virtual FColor GetTypeColor() const override { return FColor(190, 20, 70); }
+	virtual UClass* GetSupportedClass() const override { return UPreviewMeshCollection::StaticClass(); }
+	virtual uint32 GetCategories() override { return EAssetTypeCategories::Misc; }
+};
diff --git a/Engine/Plugins/Editor/EngineAssetDefinitions/Source/Private/AssetDefinition_ProceduralFoliageSpawner.h b/Engine/Plugins/Editor/EngineAssetDefinitions/Source/Private/AssetDefinition_ProceduralFoliageSpawner.h
new file mode 100644
--- /dev/null
+++ b/Engine/Plugins/Editor/EngineAssetDefinitions/Source/Private/AssetDefinition_ProceduralFoliageSpawner.h
@@ -0,0 +1,31 @@
+// Copyright Epic Games, Inc. All Rights Reserved.
+
+#pragma once
+
+#include "CoreMinimal.h"
+#include "ProceduralFoliageSpawner.h"
+#include "AssetDefinitionDefault.h"
+
+#include "AssetDefinition_ProceduralFoliageSpawner.generated.h"
+
+UCLASS()
+class UAssetDefinition_ProceduralFoliageSpawner : public UAssetDefinitionDefault
+{
+	GENERATED_BODY()
+
+public:
+	// UAssetDefinition Begin
+	virtual FText GetAssetDisplayName() const override { return NSLOCTEXT("AssetTypeActions", "AssetTypeActions_ProceduralFoliageSpawner", "Procedural Foliage Spawner"); }
+	virtual FLinearColor GetAssetColor() const override { return FLinearColor(FColor(7, 103, 7)); }
+	virtual TSoftClassPtr<UObject> GetAssetClass() const override { return UProceduralFoliageSpawner::StaticClass(); }
+	virtual TConstArrayView<FAssetCategoryPath> GetAssetCategories() const override
+	{
+		static const auto Categories = { EAssetCategoryPaths::Foliage };
+		return Categories;
+	}
+	//bool FAssetTypeActions_ProceduralFoliageSpawner::CanFilter()
+	//{
+	//	return GetDefault<UEditorExperimentalSettings>()->bProceduralFoliage;
+	//}
+	// UAssetDefinition End
+};
diff --git a/Engine/Plugins/Editor/EngineAssetDefinitions/Source/Private/AssetDefinition_Rig.h b/Engine/Plugins/Editor/EngineAssetDefinitions/Source/Private/AssetDefinition_Rig.h
new file mode 100644
--- /dev/null
+++ b/Engine/Plugins/Editor/EngineAssetDefinitions/Source/Private/AssetDefinition_Rig.h
@@ -0,0 +1,27 @@
+// Copyright Epic Games, Inc. All Rights Reserved.
+
+#pragma once
+
+#include "CoreMinimal.h"
+#include "Animation/Rig.h"
+#include "AssetDefinitionDefault.h"
+
+#include "AssetDefinition_Rig.generated.h"
+
+UCLASS()
+class UAssetDefinition_Rig : public UAssetDefinitionDefault
+{
+	GENERATED_BODY()
+
+public:
+	// UAssetDefinition Begin
+	virtual FText GetAssetDisplayName() const override { return NSLOCTEXT("AssetTypeActions", "AssetTypeActions_Rig", "Rig"); }
+	virtual FLinearColor GetAssetColor() const override { return FLinearColor(FColor(255,201,14)); }
+	virtual TSoftClassPtr<UObject> GetAssetClass() const override { return URig::StaticClass(); }
+	virtual TConstArrayView<FAssetCategoryPath> GetAssetCategories() const override
+	{
+		static const auto Categories = { EAssetCategoryPaths::Animation };
+		return Categories;
+	}
+	// UAssetDefinition End
+};
diff --git a/Engine/Plugins/Editor/EngineAssetDefinitions/Source/Private/AssetDefinition_SubUVAnimation.h b/Engine/Plugins/Editor/EngineAssetDefinitions/Source/Private/AssetDefinition_SubUVAnimation.h
new file mode 100644
--- /dev/null
+++ b/Engine/Plugins/Editor/EngineAssetDefinitions/Source/Private/AssetDefinition_SubUVAnimation.h
@@ -0,0 +1,22 @@
+// Copyright Epic Games, Inc. All Rights Reserved.
+
+#pragma once
+
+#include "CoreMinimal.h"
+#include "AssetDefinitionDefault.h"
+#include "Particles/SubUVAnimation.h"
+
+#include "AssetDefinition_SubUVAnimation.generated.h"
+
+UCLASS()
+class UAssetDefinition_SubUVAnimation : public UAssetDefinitionDefault
+{
+	GENERATED_BODY()
+
+public:
+	// UAssetDefinition Begin
+	virtual FText GetAssetDisplayName() const override { return NSLOCTEXT("AssetTypeActions", "AssetTypeActions_SubUVAnimation", "Sub UV Animation"); }
+	virtual FLinearColor GetAssetColor() const override { return FLinearColor(FColor(255,255,255)); }
+	virtual TSoftClassPtr<UObject> GetAssetClass() const override { return USubUVAnimation::StaticClass(); }
+	// UAssetDefinition End
+};
diff --git a/Engine/Plugins/Editor/EngineAssetDefinitions/Source/Private/AssetDefinition_SubsurfaceProfile.h b/Engine/Plugins/Editor/EngineAssetDefinitions/Source/Private/AssetDefinition_SubsurfaceProfile.h
new file mode 100644
--- /dev/null
+++ b/Engine/Plugins/Editor/EngineAssetDefinitions/Source/Private/AssetDefinition_SubsurfaceProfile.h
@@ -0,0 +1,27 @@
+// Copyright Epic Games, Inc. All Rights Reserved.
+
+#pragma once
+
+#include "CoreMinimal.h"
+#include "Engine/SubsurfaceProfile.h"
+#include "AssetDefinitionDefault.h"
+
+#include "AssetDefinition_SubsurfaceProfile.generated.h"
+
+UCLASS()
+class UAssetDefinition_SubsurfaceProfile : public UAssetDefinitionDefault
+{
+	GENERATED_BODY()
+
+public:
+	// UAssetDefinition Begin
+	virtual FText GetAssetDisplayName() const override { return NSLOCTEXT("AssetTypeActions", "AssetTypeActions_SubsurfaceProfile", "Subsurface Profile"); }
+	virtual FLinearColor GetAssetColor() const override { return FLinearColor(FColor(255, 196, 128)); }
+	virtual TSoftClassPtr<UObject> GetAssetClass() const override { return USubsurfaceProfile::StaticClass(); }
+	virtual TConstArrayView<FAssetCategoryPath> GetAssetCategories() const override
+    {
+    	static const auto Categories = { EAssetCategoryPaths::Material };
+    	return Categories;
+    }
+	// UAssetDefinition End
+};
diff --git a/Engine/Plugins/Editor/EngineAssetDefinitions/Source/Private/AssetDefinition_TouchInterface.h b/Engine/Plugins/Editor/EngineAssetDefinitions/Source/Private/AssetDefinition_TouchInterface.h
new file mode 100644
--- /dev/null
+++ b/Engine/Plugins/Editor/EngineAssetDefinitions/Source/Private/AssetDefinition_TouchInterface.h
@@ -0,0 +1,22 @@
+// Copyright Epic Games, Inc. All Rights Reserved.
+
+#pragma once
+
+#include "CoreMinimal.h"
+#include "GameFramework/TouchInterface.h"
+#include "AssetDefinitionDefault.h"
+
+#include "AssetDefinition_TouchInterface.generated.h"
+
+UCLASS()
+class UAssetDefinition_TouchInterface : public UAssetDefinitionDefault
+{
+	GENERATED_BODY()
+
+public:
+	// UAssetDefinition Begin
+	virtual FText GetAssetDisplayName() const override { return NSLOCTEXT("AssetTypeActions", "AssetTypeActions_TouchInterface", "Touch Interface Setup"); }
+	virtual FLinearColor GetAssetColor() const override { return FLinearColor(FColor(175, 0, 128)); }
+	virtual TSoftClassPtr<UObject> GetAssetClass() const override { return UTouchInterface::StaticClass(); }
+	// UAssetDefinition End
+};
diff --git a/Engine/Plugins/Editor/EngineAssetDefinitions/Source/Private/AssetDefinition_VectorField.h b/Engine/Plugins/Editor/EngineAssetDefinitions/Source/Private/AssetDefinition_VectorField.h
new file mode 100644
--- /dev/null
+++ b/Engine/Plugins/Editor/EngineAssetDefinitions/Source/Private/AssetDefinition_VectorField.h
@@ -0,0 +1,22 @@
+// Copyright Epic Games, Inc. All Rights Reserved.
+
+#pragma once
+
+#include "CoreMinimal.h"
+#include "VectorField/VectorField.h"
+#include "AssetDefinitionDefault.h"
+
+#include "AssetDefinition_VectorField.generated.h"
+
+UCLASS()
+class UAssetDefinition_VectorField : public UAssetDefinitionDefault
+{
+	GENERATED_BODY()
+
+public:
+	// UAssetDefinition Begin
+	virtual FText GetAssetDisplayName() const override { return NSLOCTEXT("AssetTypeActions", "AssetTypeActions_VectorField", "Vector Field"); }
+	virtual FLinearColor GetAssetColor() const override { return FLinearColor(FColor(200,128,128)); }
+	virtual TSoftClassPtr<UObject> GetAssetClass() const override { return UVectorField::StaticClass(); }
+	// UAssetDefinition End
+};
diff --git a/Engine/Plugins/Editor/EngineAssetDefinitions/Source/Private/AssetDefinition_VectorFieldAnimated.h b/Engine/Plugins/Editor/EngineAssetDefinitions/Source/Private/AssetDefinition_VectorFieldAnimated.h
new file mode 100644
--- /dev/null
+++ b/Engine/Plugins/Editor/EngineAssetDefinitions/Source/Private/AssetDefinition_VectorFieldAnimated.h
@@ -0,0 +1,21 @@
+// Copyright Epic Games, Inc. All Rights Reserved.
+
+#pragma once
+
+#include "CoreMinimal.h"
+#include "VectorField/VectorFieldAnimated.h"
+#include "AssetDefinitionDefault.h"
+
+#include "AssetDefinition_VectorFieldAnimated.generated.h"
+
+UCLASS()
+class UAssetDefinition_VectorFieldAnimated : public UAssetDefinition_VectorField
+{
+	GENERATED_BODY()
+
+public:
+	// UAssetDefinition Begin
+	virtual FText GetAssetDisplayName() const override { return NSLOCTEXT("AssetTypeActions", "AssetTypeActions_VectorFieldAnimated", "Animated Vector Field"); }
+	virtual TSoftClassPtr<UObject> GetAssetClass() const override { return UVectorFieldAnimated::StaticClass(); }
+	// UAssetDefinition End
+};
diff --git a/Engine/Plugins/Editor/EngineAssetDefinitions/Source/Private/AssetDefinition_VectorFieldStatic.h b/Engine/Plugins/Editor/EngineAssetDefinitions/Source/Private/AssetDefinition_VectorFieldStatic.h
new file mode 100644
--- /dev/null
+++ b/Engine/Plugins/Editor/EngineAssetDefinitions/Source/Private/AssetDefinition_VectorFieldStatic.h
@@ -0,0 +1,22 @@
+// Copyright Epic Games, Inc. All Rights Reserved.
+
+#pragma once
+
+#include "CoreMinimal.h"
+#include "VectorField/VectorFieldStatic.h"
+#include "AssetDefinitionDefault.h"
+
+#include "AssetDefinition_VectorFieldStatic.generated.h"
+
+UCLASS()
+class UAssetDefinition_VectorFieldStatic : public UAssetDefinition_VectorField
+{
+	GENERATED_BODY()
+
+public:
+	// UAssetDefinition Begin
+	virtual FText GetAssetDisplayName() const override { return NSLOCTEXT("AssetTypeActions", "AssetTypeActions_VectorFieldStatic", "Static Vector Field"); }
+	virtual TSoftClassPtr<UObject> GetAssetClass() const override { return UVectorFieldStatic::StaticClass(); }
+	virtual bool CanImport() const override { return true; }
+	// UAssetDefinition End
+};
diff --git a/Engine/Source/Developer/AssetTools/Private/AssetTools.cpp b/Engine/Source/Developer/AssetTools/Private/AssetTools.cpp
--- a/Engine/Source/Developer/AssetTools/Private/AssetTools.cpp
+++ b/Engine/Source/Developer/AssetTools/Private/AssetTools.cpp
@@ -42,7 +42,6 @@
 #include "AssetTypeActions/AssetTypeActions_MaterialInterface.h"
 #include "AssetTypeActions/AssetTypeActions_SkeletalMesh.h"
 #include "AssetTypeActions/AssetTypeActions_FbxSceneImportData.h"
-#include "AssetTypeActions/AssetTypeActions_VectorField.h"
 #include "AssetTypeActions/AssetTypeActions_AnimationAsset.h"
 #include "AssetTypeActions/AssetTypeActions_AnimBlueprint.h"
 #include "AssetTypeActions/AssetTypeActions_AnimBlueprintInterface.h"
@@ -70,7 +69,6 @@
 #include "AssetTypeActions/AssetTypeActions_ForceFeedbackEffect.h"
 #include "AssetTypeActions/AssetTypeActions_HapticFeedback.h"
 #include "AssetTypeActions/AssetTypeActions_HLODProxy.h"
-#include "AssetTypeActions/AssetTypeActions_SubsurfaceProfile.h"
 #include "AssetTypeActions/AssetTypeActions_ActorFoliageSettings.h"
 #include "AssetTypeActions/AssetTypeActions_InstancedFoliageSettings.h"
 #include "AssetTypeActions/AssetTypeActions_LandscapeLayer.h"
@@ -81,24 +79,12 @@
 #include "AssetTypeActions/AssetTypeActions_MaterialFunctionInstance.h"
 #include "AssetTypeActions/AssetTypeActions_MaterialInstanceConstant.h"
 #include "AssetTypeActions/AssetTypeActions_MaterialInstanceDynamic.h"
-#include "AssetTypeActions/AssetTypeActions_MaterialParameterCollection.h"
 #include "AssetTypeActions/AssetTypeActions_MirrorDataTable.h"
-#include "AssetTypeActions/AssetTypeActions_ObjectLibrary.h"
 #include "AssetTypeActions/AssetTypeActions_ParticleSystem.h"
-#include "AssetTypeActions/AssetTypeActions_PhysicalMaterial.h"
 #include "AssetTypeActions/AssetTypeActions_PhysicalMaterialMask.h"
-#include "AssetTypeActions/AssetTypeActions_PhysicsAsset.h"
-#include "AssetTypeActions/AssetTypeActions_PoseAsset.h"
-#include "AssetTypeActions/AssetTypeActions_PreviewMeshCollection.h"
-#include "AssetTypeActions/AssetTypeActions_ProceduralFoliageSpawner.h"
-#include "AssetTypeActions/AssetTypeActions_Rig.h"
 #include "AssetTypeActions/AssetTypeActions_Skeleton.h"
 #include "AssetTypeActions/AssetTypeActions_SlateBrush.h"
 #include "AssetTypeActions/AssetTypeActions_SlateWidgetStyle.h"
-#include "AssetTypeActions/AssetTypeActions_SubUVAnimation.h"
-#include "AssetTypeActions/AssetTypeActions_TouchInterface.h"
-#include "AssetTypeActions/AssetTypeActions_VectorFieldAnimated.h"
-#include "AssetTypeActions/AssetTypeActions_VectorFieldStatic.h"
 #include "WorldPartition/WorldPartition.h"
 #include "SDiscoveringAssetsDialog.h"
 #include "AssetFixUpRedirectors.h"
@@ -1238,7 +1224,6 @@ UAssetToolsImpl::UAssetToolsImpl(const FObjectInitializer& ObjectInitializer)
 	RegisterAssetTypeActions(MakeShareable(new FAssetTypeActions_AimOffset));
 	RegisterAssetTypeActions(MakeShareable(new FAssetTypeActions_AimOffset1D));
 	RegisterAssetTypeActions(MakeShareable(new FAssetTypeActions_BlendSpace));
-	RegisterAssetTypeActions(MakeShareable(new FAssetTypeActions_PoseAsset));
 	RegisterAssetTypeActions(MakeShareable(new FAssetTypeActions_BlendSpace1D));
 	RegisterAssetTypeActions(MakeShareable(new FAssetTypeActions_Blueprint));
 	RegisterAssetTypeActions(MakeShareable(new FAssetTypeActions_BlueprintGeneratedClass));
@@ -1254,7 +1239,6 @@ UAssetToolsImpl::UAssetToolsImpl(const FObjectInitializer& ObjectInitializer)
 	RegisterAssetTypeActions(MakeShareable(new FAssetTypeActions_ForceFeedbackAttenuation(InputCategoryBit)));
 	RegisterAssetTypeActions(MakeShareable(new FAssetTypeActions_ForceFeedbackEffect(InputCategoryBit)));
 	RegisterAssetTypeActions(MakeShareable(new FAssetTypeActions_HLODProxy));
-	RegisterAssetTypeActions(MakeShareable(new FAssetTypeActions_SubsurfaceProfile));
 	RegisterAssetTypeActions(MakeShareable(new FAssetTypeActions_InstancedFoliageSettings(FoliageCategoryBit)));
 	RegisterAssetTypeActions(MakeShareable(new FAssetTypeActions_ActorFoliageSettings(FoliageCategoryBit)));
 	RegisterAssetTypeActions(MakeShareable(new FAssetTypeActions_LandscapeLayer));
@@ -1270,25 +1254,13 @@ UAssetToolsImpl::UAssetToolsImpl(const FObjectInitializer& ObjectInitializer)
 	RegisterAssetTypeActions(MakeShareable(new FAssetTypeActions_MaterialInstanceConstant(EAssetTypeCategories::Type::None)));
 	RegisterAssetTypeActions(MakeShareable(new FAssetTypeActions_MaterialInstanceDynamic));
 	RegisterAssetTypeActions(MakeShareable(new FAssetTypeActions_MaterialInterface));
-	RegisterAssetTypeActions(MakeShareable(new FAssetTypeActions_MaterialParameterCollection));
 	RegisterAssetTypeActions(MakeShareable(new FAssetTypeActions_MirrorDataTable));
-	RegisterAssetTypeActions(MakeShareable(new FAssetTypeActions_ObjectLibrary));
 	RegisterAssetTypeActions(MakeShareable(new FAssetTypeActions_ParticleSystem));
-	RegisterAssetTypeActions(MakeShareable(new FAssetTypeActions_SubUVAnimation));
-	RegisterAssetTypeActions(MakeShareable(new FAssetTypeActions_PhysicalMaterial));
 	RegisterAssetTypeActions(MakeShareable(new FAssetTypeActions_PhysicalMaterialMask));
-	RegisterAssetTypeActions(MakeShareable(new FAssetTypeActions_PhysicsAsset));
-	RegisterAssetTypeActions(MakeShareable(new FAssetTypeActions_PreviewMeshCollection));
-	RegisterAssetTypeActions(MakeShareable(new FAssetTypeActions_ProceduralFoliageSpawner(FoliageCategoryBit)));
-	RegisterAssetTypeActions(MakeShareable(new FAssetTypeActions_Rig));
 	RegisterAssetTypeActions(MakeShareable(new FAssetTypeActions_SkeletalMesh));
 	RegisterAssetTypeActions(MakeShareable(new FAssetTypeActions_Skeleton));
 	RegisterAssetTypeActions(MakeShareable(new FAssetTypeActions_SlateBrush));
 	RegisterAssetTypeActions(MakeShareable(new FAssetTypeActions_SlateWidgetStyle));
-	RegisterAssetTypeActions(MakeShareable(new FAssetTypeActions_TouchInterface));
-	RegisterAssetTypeActions(MakeShareable(new FAssetTypeActions_VectorField));
-	RegisterAssetTypeActions(MakeShareable(new FAssetTypeActions_VectorFieldAnimated));
-	RegisterAssetTypeActions(MakeShareable(new FAssetTypeActions_VectorFieldStatic));
 	RegisterAssetTypeActions(MakeShareable(new FAssetTypeActions_HapticFeedbackEffectBuffer));
 	RegisterAssetTypeActions(MakeShareable(new FAssetTypeActions_HapticFeedbackEffectCurve));
 	RegisterAssetTypeActions(MakeShareable(new FAssetTypeActions_HapticFeedbackEffectSoundWave));
diff --git a/Engine/Source/Developer/AssetTools/Private/AssetTypeActions/AssetTypeActions_MaterialParameterCollection.h b/Engine/Source/Developer/AssetTools/Private/AssetTypeActions/AssetTypeActions_MaterialParameterCollection.h
deleted file mode 100644
diff --git a/Engine/Source/Developer/AssetTools/Private/AssetTypeActions/AssetTypeActions_NodeMappingContainer.h b/Engine/Source/Developer/AssetTools/Private/AssetTypeActions/AssetTypeActions_NodeMappingContainer.h
deleted file mode 100644
diff --git a/Engine/Source/Developer/AssetTools/Private/AssetTypeActions/AssetTypeActions_ObjectLibrary.h b/Engine/Source/Developer/AssetTools/Private/AssetTypeActions/AssetTypeActions_ObjectLibrary.h
deleted file mode 100644
diff --git a/Engine/Source/Developer/AssetTools/Private/AssetTypeActions/AssetTypeActions_PhysicalMaterial.cpp b/Engine/Source/Developer/AssetTools/Private/AssetTypeActions/AssetTypeActions_PhysicalMaterial.cpp
deleted file mode 100644
diff --git a/Engine/Source/Developer/AssetTools/Private/AssetTypeActions/AssetTypeActions_PhysicalMaterial.h b/Engine/Source/Developer/AssetTools/Private/AssetTypeActions/AssetTypeActions_PhysicalMaterial.h
deleted file mode 100644
diff --git a/Engine/Source/Developer/AssetTools/Private/AssetTypeActions/AssetTypeActions_PhysicsAsset.cpp b/Engine/Source/Developer/AssetTools/Private/AssetTypeActions/AssetTypeActions_PhysicsAsset.cpp
deleted file mode 100644
diff --git a/Engine/Source/Developer/AssetTools/Private/AssetTypeActions/AssetTypeActions_PhysicsAsset.h b/Engine/Source/Developer/AssetTools/Private/AssetTypeActions/AssetTypeActions_PhysicsAsset.h
deleted file mode 100644
diff --git a/Engine/Source/Developer/AssetTools/Private/AssetTypeActions/AssetTypeActions_PoseAsset.h b/Engine/Source/Developer/AssetTools/Private/AssetTypeActions/AssetTypeActions_PoseAsset.h
deleted file mode 100644
diff --git a/Engine/Source/Developer/AssetTools/Private/AssetTypeActions/AssetTypeActions_PreviewMeshCollection.h b/Engine/Source/Developer/AssetTools/Private/AssetTypeActions/AssetTypeActions_PreviewMeshCollection.h
deleted file mode 100644
diff --git a/Engine/Source/Developer/AssetTools/Private/AssetTypeActions/AssetTypeActions_ProceduralFoliageSpawner.cpp b/Engine/Source/Developer/AssetTools/Private/AssetTypeActions/AssetTypeActions_ProceduralFoliageSpawner.cpp
deleted file mode 100644
diff --git a/Engine/Source/Developer/AssetTools/Private/AssetTypeActions/AssetTypeActions_ProceduralFoliageSpawner.h b/Engine/Source/Developer/AssetTools/Private/AssetTypeActions/AssetTypeActions_ProceduralFoliageSpawner.h
deleted file mode 100644
diff --git a/Engine/Source/Developer/AssetTools/Private/AssetTypeActions/AssetTypeActions_Rig.cpp b/Engine/Source/Developer/AssetTools/Private/AssetTypeActions/AssetTypeActions_Rig.cpp
deleted file mode 100644
diff --git a/Engine/Source/Developer/AssetTools/Private/AssetTypeActions/AssetTypeActions_Rig.h b/Engine/Source/Developer/AssetTools/Private/AssetTypeActions/AssetTypeActions_Rig.h
deleted file mode 100644
diff --git a/Engine/Source/Developer/AssetTools/Private/AssetTypeActions/AssetTypeActions_SubUVAnimation.h b/Engine/Source/Developer/AssetTools/Private/AssetTypeActions/AssetTypeActions_SubUVAnimation.h
deleted file mode 100644
diff --git a/Engine/Source/Developer/AssetTools/Private/AssetTypeActions/AssetTypeActions_SubsurfaceProfile.cpp b/Engine/Source/Developer/AssetTools/Private/AssetTypeActions/AssetTypeActions_SubsurfaceProfile.cpp
deleted file mode 100644
diff --git a/Engine/Source/Developer/AssetTools/Private/AssetTypeActions/AssetTypeActions_SubsurfaceProfile.h b/Engine/Source/Developer/AssetTools/Private/AssetTypeActions/AssetTypeActions_SubsurfaceProfile.h
deleted file mode 100644
diff --git a/Engine/Source/Developer/AssetTools/Private/AssetTypeActions/AssetTypeActions_TouchInterface.cpp b/Engine/Source/Developer/AssetTools/Private/AssetTypeActions/AssetTypeActions_TouchInterface.cpp
deleted file mode 100644
diff --git a/Engine/Source/Developer/AssetTools/Private/AssetTypeActions/AssetTypeActions_TouchInterface.h b/Engine/Source/Developer/AssetTools/Private/AssetTypeActions/AssetTypeActions_TouchInterface.h
deleted file mode 100644
diff --git a/Engine/Source/Developer/AssetTools/Private/AssetTypeActions/AssetTypeActions_VectorField.h b/Engine/Source/Developer/AssetTools/Private/AssetTypeActions/AssetTypeActions_VectorField.h
deleted file mode 100644
diff --git a/Engine/Source/Developer/AssetTools/Private/AssetTypeActions/AssetTypeActions_VectorFieldAnimated.h b/Engine/Source/Developer/AssetTools/Private/AssetTypeActions/AssetTypeActions_VectorFieldAnimated.h
deleted file mode 100644
diff --git a/Engine/Source/Developer/AssetTools/Private/AssetTypeActions/AssetTypeActions_VectorFieldStatic.cpp b/Engine/Source/Developer/AssetTools/Private/AssetTypeActions/AssetTypeActions_VectorFieldStatic.cpp
deleted file mode 100644
diff --git a/Engine/Source/Developer/AssetTools/Private/AssetTypeActions/AssetTypeActions_VectorFieldStatic.h b/Engine/Source/Developer/AssetTools/Private/AssetTypeActions/AssetTypeActions_VectorFieldStatic.h
deleted file mode 100644
