AssetDefinition - Upgrading another slew of AssetTypeActions.

diff --git a/Engine/Plugins/Editor/EngineAssetDefinitions/Source/EngineAssetDefinitions.Build.cs b/Engine/Plugins/Editor/EngineAssetDefinitions/Source/EngineAssetDefinitions.Build.cs
--- a/Engine/Plugins/Editor/EngineAssetDefinitions/Source/EngineAssetDefinitions.Build.cs
+++ b/Engine/Plugins/Editor/EngineAssetDefinitions/Source/EngineAssetDefinitions.Build.cs
@@ -31,6 +31,7 @@ public class EngineAssetDefinitions : ModuleRules
 				"Kismet",
 				"DesktopPlatform",
 				"Foliage",
+				"Landscape",
 				"PhysicsCore",
 			}
 		);
diff --git a/Engine/Plugins/Editor/EngineAssetDefinitions/Source/Private/AssetDefinition_Actor.cpp b/Engine/Plugins/Editor/EngineAssetDefinitions/Source/Private/AssetDefinition_Actor.cpp
new file mode 100644
--- /dev/null
+++ b/Engine/Plugins/Editor/EngineAssetDefinitions/Source/Private/AssetDefinition_Actor.cpp
@@ -0,0 +1,23 @@
+// Copyright Epic Games, Inc. All Rights Reserved.
+
+#include "AssetDefinition_Actor.h"
+
+#define LOCTEXT_NAMESPACE "AssetTypeActions"
+
+EAssetCommandResult UAssetDefinition_Actor::OpenAssets(const FAssetOpenArgs& OpenArgs) const
+{
+	for (const FAssetData& Asset : OpenArgs.Assets)
+	{
+		if (AActor* Actor = Cast<AActor>(Asset.GetAsset()))
+		{
+			if (Actor->OpenAssetEditor())
+			{
+				return EAssetCommandResult::Handled;
+			}
+		}
+	}
+
+	return EAssetCommandResult::Handled;
+}
+
+#undef LOCTEXT_NAMESPACE
diff --git a/Engine/Plugins/Editor/EngineAssetDefinitions/Source/Private/AssetDefinition_Actor.h b/Engine/Plugins/Editor/EngineAssetDefinitions/Source/Private/AssetDefinition_Actor.h
new file mode 100644
--- /dev/null
+++ b/Engine/Plugins/Editor/EngineAssetDefinitions/Source/Private/AssetDefinition_Actor.h
@@ -0,0 +1,29 @@
+// Copyright Epic Games, Inc. All Rights Reserved.
+
+#pragma once
+
+#include "CoreMinimal.h"
+#include "GameFramework/Actor.h"
+#include "AssetDefinitionDefault.h"
+
+#include "AssetDefinition_Actor.generated.h"
+
+UCLASS()
+class UAssetDefinition_Actor : public UAssetDefinitionDefault
+{
+	GENERATED_BODY()
+
+public:
+	// UAssetDefinition Begin
+	virtual FText GetAssetDisplayName() const override { return NSLOCTEXT("AssetTypeActions", "AssetTypeActions_Actor", "Actor"); }
+	virtual FLinearColor GetAssetColor() const override { return FLinearColor(FColor(0,232,0)); }
+	virtual TSoftClassPtr<UObject> GetAssetClass() const override { return AActor::StaticClass(); }
+	virtual TConstArrayView<FAssetCategoryPath> GetAssetCategories() const override
+	{
+		return TConstArrayView<FAssetCategoryPath>();
+	}
+	virtual FText GetObjectDisplayNameText(UObject* Object) const override { return FText::FromString(CastChecked<AActor>(Object)->GetActorLabel()); }
+
+	virtual EAssetCommandResult OpenAssets(const FAssetOpenArgs& OpenArgs) const override;
+	// UAssetDefinition End
+};
diff --git a/Engine/Plugins/Editor/EngineAssetDefinitions/Source/Private/AssetDefinition_ActorFolder.h b/Engine/Plugins/Editor/EngineAssetDefinitions/Source/Private/AssetDefinition_ActorFolder.h
new file mode 100644
--- /dev/null
+++ b/Engine/Plugins/Editor/EngineAssetDefinitions/Source/Private/AssetDefinition_ActorFolder.h
@@ -0,0 +1,30 @@
+// Copyright Epic Games, Inc. All Rights Reserved.
+
+#pragma once
+
+#include "ActorFolder.h"
+#include "AssetDefinitionDefault.h"
+
+#include "AssetDefinition_ActorFolder.generated.h"
+
+UCLASS()
+class UAssetDefinition_ActorFolder : public UAssetDefinitionDefault
+{
+	GENERATED_BODY()
+
+public:
+	// UAssetDefinition Begin
+	virtual FText GetAssetDisplayName() const override { return NSLOCTEXT("AssetTypeActions", "AssetTypeActions_ActorFolder", "Actor Folder"); }
+	virtual FLinearColor GetAssetColor() const override { return FLinearColor(FColor(182, 143, 85)); }
+	virtual TSoftClassPtr<UObject> GetAssetClass() const override { return UActorFolder::StaticClass(); }
+	virtual TConstArrayView<FAssetCategoryPath> GetAssetCategories() const override
+	{
+		return TConstArrayView<FAssetCategoryPath>();
+	}
+	virtual FText GetObjectDisplayNameText(UObject* Object) const override
+	{
+		UActorFolder* ActorFolder = CastChecked<UActorFolder>(Object);
+		return FText::FromString(ActorFolder->GetDisplayName());
+	}
+	// UAssetDefinition End
+};
diff --git a/Engine/Plugins/Editor/EngineAssetDefinitions/Source/Private/AssetDefinition_DataLayer.h b/Engine/Plugins/Editor/EngineAssetDefinitions/Source/Private/AssetDefinition_DataLayer.h
new file mode 100644
--- /dev/null
+++ b/Engine/Plugins/Editor/EngineAssetDefinitions/Source/Private/AssetDefinition_DataLayer.h
@@ -0,0 +1,22 @@
+// Copyright Epic Games, Inc. All Rights Reserved.
+
+#pragma once
+
+#include "CoreMinimal.h"
+#include "WorldPartition/DataLayer/DataLayerAsset.h"
+#include "AssetDefinitionDefault.h"
+
+#include "AssetDefinition_DataLayer.generated.h"
+
+UCLASS()
+class UAssetDefinition_DataLayer : public UAssetDefinitionDefault
+{
+	GENERATED_BODY()
+
+public:
+	// UAssetDefinition Begin
+	virtual FText GetAssetDisplayName() const override { return NSLOCTEXT("AssetTypeActions", "AssetTypeActions_DataLayer", "Data Layer"); }
+	virtual FLinearColor GetAssetColor() const override { return FLinearColor(FColor(12, 65, 12)); }
+	virtual TSoftClassPtr<UObject> GetAssetClass() const override { return UDataLayerAsset::StaticClass(); }
+	// UAssetDefinition End
+};
diff --git a/Engine/Plugins/Editor/EngineAssetDefinitions/Source/Private/AssetDefinition_FbxSceneImportData.h b/Engine/Plugins/Editor/EngineAssetDefinitions/Source/Private/AssetDefinition_FbxSceneImportData.h
new file mode 100644
--- /dev/null
+++ b/Engine/Plugins/Editor/EngineAssetDefinitions/Source/Private/AssetDefinition_FbxSceneImportData.h
@@ -0,0 +1,23 @@
+// Copyright Epic Games, Inc. All Rights Reserved.
+
+#pragma once
+
+#include "CoreMinimal.h"
+#include "Factories/FbxSceneImportData.h"
+#include "AssetDefinitionDefault.h"
+
+#include "AssetDefinition_FbxSceneImportData.generated.h"
+
+UCLASS()
+class UAssetDefinition_FbxSceneImportData : public UAssetDefinitionDefault
+{
+	GENERATED_BODY()
+
+public:
+	// UAssetDefinition Begin
+	virtual FText GetAssetDisplayName() const override { return NSLOCTEXT("AssetTypeActions", "AssetTypeActions_SceneImportData", "Fbx Scene Import Data"); }
+	virtual FLinearColor GetAssetColor() const override { return FLinearColor(FColor(255,0,255)); }
+	virtual TSoftClassPtr<UObject> GetAssetClass() const override { return UFbxSceneImportData::StaticClass(); }
+	virtual bool CanImport() const override { return true; }
+	// UAssetDefinition End
+};
diff --git a/Engine/Plugins/Editor/EngineAssetDefinitions/Source/Private/AssetDefinition_ForceFeedbackAttenuation.h b/Engine/Plugins/Editor/EngineAssetDefinitions/Source/Private/AssetDefinition_ForceFeedbackAttenuation.h
new file mode 100644
--- /dev/null
+++ b/Engine/Plugins/Editor/EngineAssetDefinitions/Source/Private/AssetDefinition_ForceFeedbackAttenuation.h
@@ -0,0 +1,27 @@
+// Copyright Epic Games, Inc. All Rights Reserved.
+
+#pragma once
+
+#include "CoreMinimal.h"
+#include "GameFramework/ForceFeedbackAttenuation.h"
+#include "AssetDefinitionDefault.h"
+
+#include "AssetDefinition_ForceFeedbackAttenuation.generated.h"
+
+UCLASS()
+class UAssetDefinition_ForceFeedbackAttenuation : public UAssetDefinitionDefault
+{
+	GENERATED_BODY()
+
+public:
+	// UAssetDefinition Begin
+	virtual FText GetAssetDisplayName() const override { return NSLOCTEXT("AssetTypeActions", "AssetTypeActions_ForceFeedbackAttenuation", "Force Feedback Attenuation"); }
+	virtual FLinearColor GetAssetColor() const override { return FLinearColor(FColor(175, 0, 0)); }
+	virtual TSoftClassPtr<UObject> GetAssetClass() const override { return UForceFeedbackAttenuation::StaticClass(); }
+	virtual TConstArrayView<FAssetCategoryPath> GetAssetCategories() const override
+   	{
+   		static const auto Categories = { EAssetCategoryPaths::Input };
+   		return Categories;
+   	}
+	// UAssetDefinition End
+};
diff --git a/Engine/Plugins/Editor/EngineAssetDefinitions/Source/Private/AssetDefinition_HLODProxy.h b/Engine/Plugins/Editor/EngineAssetDefinitions/Source/Private/AssetDefinition_HLODProxy.h
new file mode 100644
--- /dev/null
+++ b/Engine/Plugins/Editor/EngineAssetDefinitions/Source/Private/AssetDefinition_HLODProxy.h
@@ -0,0 +1,23 @@
+// Copyright Epic Games, Inc. All Rights Reserved.
+
+#pragma once
+
+#include "CoreMinimal.h"
+#include "Engine/HLODProxy.h"
+#include "AssetDefinitionDefault.h"
+
+#include "AssetDefinition_HLODProxy.generated.h"
+
+UCLASS()
+class UAssetDefinition_HLODProxy : public UAssetDefinitionDefault
+{
+	GENERATED_BODY()
+
+public:
+	// UAssetDefinition Begin
+	virtual FText GetAssetDisplayName() const override { return NSLOCTEXT("AssetTypeActions", "AssetTypeActions_HLODProxy", "HLOD Proxy"); }
+	virtual FLinearColor GetAssetColor() const override { return FLinearColor(FColor(0, 200, 200)); }
+	virtual TSoftClassPtr<UObject> GetAssetClass() const override { return UHLODProxy::StaticClass(); }
+	virtual FAssetSupportResponse CanLocalize(const FAssetData& InAsset) const override { return FAssetSupportResponse::NotSupported(); }
+	// UAssetDefinition End
+};
diff --git a/Engine/Plugins/Editor/EngineAssetDefinitions/Source/Private/AssetDefinition_InstancedFoliageSettings.h b/Engine/Plugins/Editor/EngineAssetDefinitions/Source/Private/AssetDefinition_InstancedFoliageSettings.h
new file mode 100644
--- /dev/null
+++ b/Engine/Plugins/Editor/EngineAssetDefinitions/Source/Private/AssetDefinition_InstancedFoliageSettings.h
@@ -0,0 +1,27 @@
+// Copyright Epic Games, Inc. All Rights Reserved.
+
+#pragma once
+
+#include "CoreMinimal.h"
+#include "AssetDefinitionDefault.h"
+#include "FoliageType_InstancedStaticMesh.h"
+
+#include "AssetDefinition_InstancedFoliageSettings.generated.h"
+
+UCLASS()
+class UAssetDefinition_InstancedFoliageSettings : public UAssetDefinitionDefault
+{
+	GENERATED_BODY()
+
+public:
+	// UAssetDefinition Begin
+	virtual FText GetAssetDisplayName() const override { return NSLOCTEXT("AssetTypeActions", "AssetTypeActions_InstancedFoliageSettings", "Static Mesh Foliage"); }
+	virtual FLinearColor GetAssetColor() const override { return FLinearColor(FColor(12, 173, 12)); }
+	virtual TSoftClassPtr<UObject> GetAssetClass() const override { return UFoliageType_InstancedStaticMesh::StaticClass(); }
+	virtual TConstArrayView<FAssetCategoryPath> GetAssetCategories() const override
+	{
+		static const auto Categories = { EAssetCategoryPaths::Foliage };
+		return Categories;
+	}
+	// UAssetDefinition End
+};
diff --git a/Engine/Plugins/Editor/EngineAssetDefinitions/Source/Private/AssetDefinition_LandscapeGrassType.h b/Engine/Plugins/Editor/EngineAssetDefinitions/Source/Private/AssetDefinition_LandscapeGrassType.h
new file mode 100644
--- /dev/null
+++ b/Engine/Plugins/Editor/EngineAssetDefinitions/Source/Private/AssetDefinition_LandscapeGrassType.h
@@ -0,0 +1,27 @@
+// Copyright Epic Games, Inc. All Rights Reserved.
+
+#pragma once
+
+#include "CoreMinimal.h"
+#include "LandscapeGrassType.h"
+#include "AssetDefinitionDefault.h"
+
+#include "AssetDefinition_LandscapeGrassType.generated.h"
+
+UCLASS()
+class UAssetDefinition_LandscapeGrassType : public UAssetDefinitionDefault
+{
+	GENERATED_BODY()
+
+public:
+	// UAssetDefinition Begin
+	virtual FText GetAssetDisplayName() const override { return NSLOCTEXT("AssetTypeActions", "AssetTypeActions_LandscapeGrassType", "Landscape Grass Type"); }
+	virtual FLinearColor GetAssetColor() const override { return FLinearColor(FColor(128,255,128)); }
+	virtual TSoftClassPtr<UObject> GetAssetClass() const override { return ULandscapeGrassType::StaticClass(); }
+	virtual TConstArrayView<FAssetCategoryPath> GetAssetCategories() const override
+	{
+		static const auto Categories = { EAssetCategoryPaths::Foliage };
+		return Categories;
+	}
+	// UAssetDefinition End
+};
diff --git a/Engine/Plugins/Editor/EngineAssetDefinitions/Source/Private/AssetDefinition_LandscapeLayer.h b/Engine/Plugins/Editor/EngineAssetDefinitions/Source/Private/AssetDefinition_LandscapeLayer.h
new file mode 100644
--- /dev/null
+++ b/Engine/Plugins/Editor/EngineAssetDefinitions/Source/Private/AssetDefinition_LandscapeLayer.h
@@ -0,0 +1,22 @@
+// Copyright Epic Games, Inc. All Rights Reserved.
+
+#pragma once
+
+#include "CoreMinimal.h"
+#include "LandscapeLayerInfoObject.h"
+#include "AssetDefinitionDefault.h"
+
+#include "AssetDefinition_LandscapeLayer.generated.h"
+
+UCLASS()
+class UAssetDefinition_LandscapeLayer : public UAssetDefinitionDefault
+{
+	GENERATED_BODY()
+
+public:
+	// UAssetDefinition Begin
+	virtual FText GetAssetDisplayName() const override { return NSLOCTEXT("AssetTypeActions", "AssetTypeActions_LandscapeLayer", "Landscape Layer"); }
+	virtual FLinearColor GetAssetColor() const override { return FLinearColor(FColor(128,192,255)); }
+	virtual TSoftClassPtr<UObject> GetAssetClass() const override { return ULandscapeLayerInfoObject::StaticClass(); }
+	// UAssetDefinition End
+};
diff --git a/Engine/Plugins/Editor/EngineAssetDefinitions/Source/Private/AssetDefinition_PhysicsAsset.cpp b/Engine/Plugins/Editor/EngineAssetDefinitions/Source/Private/AssetDefinition_PhysicsAsset.cpp
--- a/Engine/Plugins/Editor/EngineAssetDefinitions/Source/Private/AssetDefinition_PhysicsAsset.cpp
+++ b/Engine/Plugins/Editor/EngineAssetDefinitions/Source/Private/AssetDefinition_PhysicsAsset.cpp
@@ -11,7 +11,7 @@ UThumbnailInfo* UAssetDefinition_PhysicsAsset::LoadThumbnailInfo(const FAssetDat
 {
 	UPhysicsAsset* PhysicsAsset = CastChecked<UPhysicsAsset>(InAsset.GetAsset());
 	UThumbnailInfo* ThumbnailInfo = PhysicsAsset->ThumbnailInfo;
-	if (ThumbnailInfo == NULL)
+	if (ThumbnailInfo == nullptr)
 	{
 		ThumbnailInfo = NewObject<USceneThumbnailInfo>(PhysicsAsset, NAME_None, RF_Transactional);
 		PhysicsAsset->ThumbnailInfo = ThumbnailInfo;
diff --git a/Engine/Source/Developer/AssetTools/Private/AssetDefinition_AssetTypeActionsProxy.cpp b/Engine/Source/Developer/AssetTools/Private/AssetDefinition_AssetTypeActionsProxy.cpp
--- a/Engine/Source/Developer/AssetTools/Private/AssetDefinition_AssetTypeActionsProxy.cpp
+++ b/Engine/Source/Developer/AssetTools/Private/AssetDefinition_AssetTypeActionsProxy.cpp
@@ -205,7 +205,7 @@ EAssetCommandResult UAssetDefinition_AssetTypeActionsProxy::GetFilters(TArray<FA
 	return EAssetCommandResult::Unhandled;
 }
 
-FText UAssetDefinition_AssetTypeActionsProxy::GetObjectDisplayText(UObject* Object) const
+FText UAssetDefinition_AssetTypeActionsProxy::GetObjectDisplayNameText(UObject* Object) const
 {
 	return FText::FromString(AssetType->GetObjectDisplayName(Object));
 }
diff --git a/Engine/Source/Developer/AssetTools/Private/AssetDefinition_AssetTypeActionsProxy.h b/Engine/Source/Developer/AssetTools/Private/AssetDefinition_AssetTypeActionsProxy.h
--- a/Engine/Source/Developer/AssetTools/Private/AssetDefinition_AssetTypeActionsProxy.h
+++ b/Engine/Source/Developer/AssetTools/Private/AssetDefinition_AssetTypeActionsProxy.h
@@ -43,7 +43,7 @@ public:
 	// We don't need to implement GetSourceFiles, the default one will / should work for any existing type.
 	// virtual EAssetCommandResult GetSourceFiles(const FAssetSourceFileArgs& SourceFileArgs, TArray<FAssetSourceFile>& OutSourceAssets) const override;
 
-	virtual FText GetObjectDisplayText(UObject* Object) const override;
+	virtual FText GetObjectDisplayNameText(UObject* Object) const override;
 	virtual EAssetCommandResult PerformAssetDiff(const FAssetDiffArgs& DiffArgs) const override;
 	
 	virtual UThumbnailInfo* LoadThumbnailInfo(const FAssetData& Asset) const override;
diff --git a/Engine/Source/Developer/AssetTools/Private/AssetTools.cpp b/Engine/Source/Developer/AssetTools/Private/AssetTools.cpp
--- a/Engine/Source/Developer/AssetTools/Private/AssetTools.cpp
+++ b/Engine/Source/Developer/AssetTools/Private/AssetTools.cpp
@@ -35,13 +35,10 @@
 #include "IClassTypeActions.h"
 #include "ClassViewerModule.h"
 #include "ClassViewerFilter.h"
-#include "AssetTypeActions/AssetTypeActions_Actor.h"
-#include "AssetTypeActions/AssetTypeActions_ActorFolder.h"
 #include "AssetTypeActions/AssetTypeActions_Blueprint.h"
 #include "AssetTypeActions/AssetTypeActions_BlueprintGeneratedClass.h"
 #include "AssetTypeActions/AssetTypeActions_MaterialInterface.h"
 #include "AssetTypeActions/AssetTypeActions_SkeletalMesh.h"
-#include "AssetTypeActions/AssetTypeActions_FbxSceneImportData.h"
 #include "AssetTypeActions/AssetTypeActions_AnimationAsset.h"
 #include "AssetTypeActions/AssetTypeActions_AnimBlueprint.h"
 #include "AssetTypeActions/AssetTypeActions_AnimBlueprintInterface.h"
@@ -59,20 +56,14 @@
 #include "AssetTypeActions/AssetTypeActions_CurveTable.h"
 #include "AssetTypeActions/AssetTypeActions_CompositeCurveTable.h"
 #include "AssetTypeActions/AssetTypeActions_DataAsset.h"
-#include "AssetTypeActions/AssetTypeActions_DataLayer.h"
 #include "AssetTypeActions/AssetTypeActions_DataTable.h"
 #include "AssetTypeActions/AssetTypeActions_CompositeDataTable.h"
 #include "AssetTypeActions/AssetTypeActions_Class.h"
 #include "AssetTypeActions/AssetTypeActions_Font.h"
 #include "AssetTypeActions/AssetTypeActions_FontFace.h"
-#include "AssetTypeActions/AssetTypeActions_ForceFeedbackAttenuation.h"
 #include "AssetTypeActions/AssetTypeActions_ForceFeedbackEffect.h"
 #include "AssetTypeActions/AssetTypeActions_HapticFeedback.h"
-#include "AssetTypeActions/AssetTypeActions_HLODProxy.h"
 #include "AssetTypeActions/AssetTypeActions_ActorFoliageSettings.h"
-#include "AssetTypeActions/AssetTypeActions_InstancedFoliageSettings.h"
-#include "AssetTypeActions/AssetTypeActions_LandscapeLayer.h"
-#include "AssetTypeActions/AssetTypeActions_LandscapeGrassType.h"
 #include "AssetTypeActions/AssetTypeActions_LightWeightInstance.h"
 #include "AssetTypeActions/AssetTypeActions_Material.h"
 #include "AssetTypeActions/AssetTypeActions_MaterialFunction.h"
@@ -147,6 +138,7 @@
 #include "VirtualTexturingEditorModule.h"
 #include "Algo/AnyOf.h"
 #include "Engine/UserDefinedStruct.h"
+#include "Factories/SceneImportFactory.h"
 #include "Misc/AssetFilterData.h"
 
 #include UE_INLINE_GENERATED_CPP_BY_NAME(AssetTools)
@@ -1209,8 +1201,6 @@ UAssetToolsImpl::UAssetToolsImpl(const FObjectInitializer& ObjectInitializer)
 	EAssetTypeCategories::Type InputCategoryBit = RegisterAdvancedAssetCategory(FName(TEXT("Input")), LOCTEXT("InputAssetsCategory", "Input"));
 
 	// Register the built-in asset type actions
-	RegisterAssetTypeActions(MakeShareable(new FAssetTypeActions_Actor));
-	RegisterAssetTypeActions(MakeShareable(new FAssetTypeActions_ActorFolder));
 	RegisterAssetTypeActions(MakeShareable(new FAssetTypeActions_AnimationAsset));
 	RegisterAssetTypeActions(MakeShareable(new FAssetTypeActions_AnimBlueprintInterface));
 	RegisterAssetTypeActions(MakeShareable(new FAssetTypeActions_AnimBlueprint));
@@ -1229,20 +1219,13 @@ UAssetToolsImpl::UAssetToolsImpl(const FObjectInitializer& ObjectInitializer)
 	RegisterAssetTypeActions(MakeShareable(new FAssetTypeActions_BlueprintGeneratedClass));
 	RegisterAssetTypeActions(MakeShareable(new FAssetTypeActions_CurveTable));
 	RegisterAssetTypeActions(MakeShareable(new FAssetTypeActions_CompositeCurveTable));
-	RegisterAssetTypeActions(MakeShareable(new FAssetTypeActions_DataLayer));
 	RegisterAssetTypeActions(MakeShareable(new FAssetTypeActions_DataTable));
 	RegisterAssetTypeActions(MakeShareable(new FAssetTypeActions_CompositeDataTable));
 	RegisterAssetTypeActions(MakeShareable(new FAssetTypeActions_Class));
-	RegisterAssetTypeActions(MakeShareable(new FAssetTypeActions_SceneImportData));
 	RegisterAssetTypeActions(MakeShareable(new FAssetTypeActions_Font));
 	RegisterAssetTypeActions(MakeShareable(new FAssetTypeActions_FontFace));
-	RegisterAssetTypeActions(MakeShareable(new FAssetTypeActions_ForceFeedbackAttenuation(InputCategoryBit)));
 	RegisterAssetTypeActions(MakeShareable(new FAssetTypeActions_ForceFeedbackEffect(InputCategoryBit)));
-	RegisterAssetTypeActions(MakeShareable(new FAssetTypeActions_HLODProxy));
-	RegisterAssetTypeActions(MakeShareable(new FAssetTypeActions_InstancedFoliageSettings(FoliageCategoryBit)));
 	RegisterAssetTypeActions(MakeShareable(new FAssetTypeActions_ActorFoliageSettings(FoliageCategoryBit)));
-	RegisterAssetTypeActions(MakeShareable(new FAssetTypeActions_LandscapeLayer));
-	RegisterAssetTypeActions(MakeShareable(new FAssetTypeActions_LandscapeGrassType(FoliageCategoryBit)));
 	RegisterAssetTypeActions(MakeShareable(new FAssetTypeActions_LightWeightInstance));
 	RegisterAssetTypeActions(MakeShareable(new FAssetTypeActions_Material(EAssetTypeCategories::Type::None)));
 	RegisterAssetTypeActions(MakeShareable(new FAssetTypeActions_MaterialFunction));
diff --git a/Engine/Source/Developer/AssetTools/Private/AssetTypeActions/AssetTypeActions_Actor.cpp b/Engine/Source/Developer/AssetTools/Private/AssetTypeActions/AssetTypeActions_Actor.cpp
deleted file mode 100644
diff --git a/Engine/Source/Developer/AssetTools/Private/AssetTypeActions/AssetTypeActions_Actor.h b/Engine/Source/Developer/AssetTools/Private/AssetTypeActions/AssetTypeActions_Actor.h
deleted file mode 100644
diff --git a/Engine/Source/Developer/AssetTools/Private/AssetTypeActions/AssetTypeActions_ActorFolder.h b/Engine/Source/Developer/AssetTools/Private/AssetTypeActions/AssetTypeActions_ActorFolder.h
deleted file mode 100644
diff --git a/Engine/Source/Developer/AssetTools/Private/AssetTypeActions/AssetTypeActions_DataLayer.h b/Engine/Source/Developer/AssetTools/Private/AssetTypeActions/AssetTypeActions_DataLayer.h
deleted file mode 100644
diff --git a/Engine/Source/Developer/AssetTools/Private/AssetTypeActions/AssetTypeActions_FbxSceneImportData.cpp b/Engine/Source/Developer/AssetTools/Private/AssetTypeActions/AssetTypeActions_FbxSceneImportData.cpp
deleted file mode 100644
diff --git a/Engine/Source/Developer/AssetTools/Private/AssetTypeActions/AssetTypeActions_FbxSceneImportData.h b/Engine/Source/Developer/AssetTools/Private/AssetTypeActions/AssetTypeActions_FbxSceneImportData.h
deleted file mode 100644
diff --git a/Engine/Source/Developer/AssetTools/Private/AssetTypeActions/AssetTypeActions_ForceFeedbackAttenuation.h b/Engine/Source/Developer/AssetTools/Private/AssetTypeActions/AssetTypeActions_ForceFeedbackAttenuation.h
deleted file mode 100644
diff --git a/Engine/Source/Developer/AssetTools/Private/AssetTypeActions/AssetTypeActions_HLODProxy.cpp b/Engine/Source/Developer/AssetTools/Private/AssetTypeActions/AssetTypeActions_HLODProxy.cpp
deleted file mode 100644
diff --git a/Engine/Source/Developer/AssetTools/Private/AssetTypeActions/AssetTypeActions_HLODProxy.h b/Engine/Source/Developer/AssetTools/Private/AssetTypeActions/AssetTypeActions_HLODProxy.h
deleted file mode 100644
diff --git a/Engine/Source/Developer/AssetTools/Private/AssetTypeActions/AssetTypeActions_InstancedFoliageSettings.cpp b/Engine/Source/Developer/AssetTools/Private/AssetTypeActions/AssetTypeActions_InstancedFoliageSettings.cpp
deleted file mode 100644
diff --git a/Engine/Source/Developer/AssetTools/Private/AssetTypeActions/AssetTypeActions_InstancedFoliageSettings.h b/Engine/Source/Developer/AssetTools/Private/AssetTypeActions/AssetTypeActions_InstancedFoliageSettings.h
deleted file mode 100644
diff --git a/Engine/Source/Developer/AssetTools/Private/AssetTypeActions/AssetTypeActions_LandscapeGrassType.h b/Engine/Source/Developer/AssetTools/Private/AssetTypeActions/AssetTypeActions_LandscapeGrassType.h
deleted file mode 100644
diff --git a/Engine/Source/Developer/AssetTools/Private/AssetTypeActions/AssetTypeActions_LandscapeLayer.h b/Engine/Source/Developer/AssetTools/Private/AssetTypeActions/AssetTypeActions_LandscapeLayer.h
deleted file mode 100644
diff --git a/Engine/Source/Editor/AssetDefinition/Public/AssetDefinition.h b/Engine/Source/Editor/AssetDefinition/Public/AssetDefinition.h
--- a/Engine/Source/Editor/AssetDefinition/Public/AssetDefinition.h
+++ b/Engine/Source/Editor/AssetDefinition/Public/AssetDefinition.h
@@ -343,7 +343,7 @@ public:
 
 	
 	// Extras
-	virtual FText GetObjectDisplayText(UObject* Object) const
+	virtual FText GetObjectDisplayNameText(UObject* Object) const
 	{
 		return FText::FromString(Object->GetName());
 	}
