Fix Actor label wasn't used anymore in Save Package Dialog

diff --git a/Engine/Source/Developer/AssetTools/Private/AssetTools.cpp b/Engine/Source/Developer/AssetTools/Private/AssetTools.cpp
--- a/Engine/Source/Developer/AssetTools/Private/AssetTools.cpp
+++ b/Engine/Source/Developer/AssetTools/Private/AssetTools.cpp
@@ -162,6 +162,7 @@ public:
 
 	UAssetDefinition* GetAssetDefinition() const { return AssetDefinitionPtr.Get(); }
 	
+	virtual FString GetObjectDisplayName(UObject* Object) const override { return AssetDefinitionPtr.Get()->GetObjectDisplayNameText(Object).ToString(); }
 	virtual FText GetName() const override { return AssetDefinitionPtr.Get()->GetAssetDisplayName(); }
 	virtual UClass* GetSupportedClass() const override { return AssetDefinitionPtr.Get()->GetAssetClass().LoadSynchronous(); }
 	virtual FColor GetTypeColor() const override { return AssetDefinitionPtr.Get()->GetAssetColor().ToFColor(true); }
