Add the process of storing and restoring the camera location and rotation when source control is syncing packages
If a sync with source control reloads the map then it will perform a camera reset.
We can improve the user experience by restoring the camera position and rotation to prevent the camera moving each time a user syncs.

diff --git a/Engine/Source/Developer/AssetTools/Private/AssetViewUtils.cpp b/Engine/Source/Developer/AssetTools/Private/AssetViewUtils.cpp
--- a/Engine/Source/Developer/AssetTools/Private/AssetViewUtils.cpp
+++ b/Engine/Source/Developer/AssetTools/Private/AssetViewUtils.cpp
@@ -41,6 +41,7 @@
 #include "Subsystems/AssetEditorSubsystem.h"
 #include "Editor.h"
 #include "UObject/LinkerInstancingContext.h"
+#include "LevelEditorViewport.h"
 
 #define LOCTEXT_NAMESPACE "ContentBrowser"
 
@@ -1736,9 +1737,34 @@ void AssetViewUtils::SyncPackagesFromSourceControl(const TArray<FString>& Packag
 			return false; // keep package
 		});
 
+		// Syncing may result in a map reload, so we retain the camera views to restore after the reload...
+		struct FCameraView
+		{
+			FVector Location;
+			FRotator Rotation;
+		};
+		TMap<FLevelEditorViewportClient*, FCameraView> CameraViews;
+		for (FLevelEditorViewportClient* LevelVC : GEditor->GetLevelViewportClients())
+		{
+			CameraViews.Add(LevelVC, { LevelVC->GetViewLocation(), LevelVC->GetViewRotation() });
+		}
+
 		// Hot-reload the new packages...
 		UPackageTools::ReloadPackages(LoadedPackages);
 
+		// Restore the camera views after a map reload if necessary...
+		for (FLevelEditorViewportClient* LevelVC : GEditor->GetLevelViewportClients())
+		{
+			const FCameraView& CameraView = CameraViews.FindChecked(LevelVC);
+			LevelVC->SetViewLocation(CameraView.Location);
+			if (!LevelVC->IsOrtho())
+			{
+				LevelVC->SetViewRotation(CameraView.Rotation);
+			}
+			LevelVC->Invalidate();
+			FEditorDelegates::OnEditorCameraMoved.Broadcast(CameraView.Location, CameraView.Rotation, LevelVC->ViewportType, LevelVC->ViewIndex);
+		}
+
 		// Unload any deleted packages...
 		UPackageTools::UnloadPackages(PackagesToUnload);
 
