Fixed a issue where the engine assets couldn't be migrated. This is a regression that can block some user workflow.

diff --git a/Engine/Source/Developer/AssetTools/Private/AssetTools.cpp b/Engine/Source/Developer/AssetTools/Private/AssetTools.cpp
--- a/Engine/Source/Developer/AssetTools/Private/AssetTools.cpp
+++ b/Engine/Source/Developer/AssetTools/Private/AssetTools.cpp
@@ -4000,44 +4000,57 @@ void UAssetToolsImpl::PerformMigratePackages(TArray<FName> PackageNamesToMigrate
 
 		for (const FName& PackageName : AllPackageNamesToMove)
 		{
-			 FName PackageMountPoint = FPackageName::GetPackageMountPoint(PackageName.ToString(), false);
-			 EPluginLoadedFrom* Found = EnabledPluginToLoadedFrom.Find(PackageMountPoint);
+			FName PackageMountPoint = FPackageName::GetPackageMountPoint(PackageName.ToString(), false);
+			EPluginLoadedFrom* Found = EnabledPluginToLoadedFrom.Find(PackageMountPoint);
 
-			 bool bShouldMigratePackage = true;
-			 if (Found)
+			bool bShouldMigratePackage = true;
+			if (Found)
 			 {
-				 // plugin content, decide if it's appropriate to migrate
-				 switch (*Found)
-				 {
-				 case EPluginLoadedFrom::Engine:
-					 if (!bShouldShowEngineContent)
-					 {
-						 continue;
-					 }
-					 bShouldMigratePackage = false;
-					 break;
-
-				 case EPluginLoadedFrom::Project:
-					 bShouldMigratePackage = true;
-					 break;
+				// plugin content, decide if it's appropriate to migrate
+				switch (*Found)
+				{
+				case EPluginLoadedFrom::Engine:
+					if (!bShouldShowEngineContent)
+					{
+						continue;
+					}
+					bShouldMigratePackage = false;
+					break;
+
+				case EPluginLoadedFrom::Project:
+					bShouldMigratePackage = true;
+					break;
 				 
-				 default:
-					 bShouldMigratePackage = false;
-					 break;
+				default:
+					bShouldMigratePackage = false;
+					break;
 				 }
 			 }
 			 else
 			 {
-				 // this is not plugin content
-				 bShouldMigratePackage = true;
-			 }
+				// this is not plugin content
+				if (PackageName.ToString().StartsWith(TEXT("/Engine")))
+				{
+					// Engine content
+					if (!bShouldShowEngineContent)
+					{
+						continue;
+					}
+					bShouldMigratePackage = false;
+				}
+				else
+				{
+					// Game content
+					bShouldMigratePackage = true;
+				}
+			}
 
-			 FilteredPackageNamesToMove.Add(PackageName);
+			FilteredPackageNamesToMove.Add(PackageName);
 
-			 if (bShouldMigratePackage)
-			 {
-				 ShouldMigratePackage.Add(PackageName);
-			 }
+			if (bShouldMigratePackage)
+			{
+				ShouldMigratePackage.Add(PackageName);
+			}
 		}
 
 		AllPackageNamesToMove = FilteredPackageNamesToMove;
@@ -4841,13 +4854,12 @@ void UAssetToolsImpl::RecursiveGetDependencies(const FName& PackageName, TSet<FN
 	{
 		FString DependencyName = (*DependsIt).ToString();
 
-		const bool bIsEnginePackage = DependencyName.StartsWith(TEXT("/Engine"));
 		const bool bIsScriptPackage = DependencyName.StartsWith(TEXT("/Script"));
 
 		// The asset registry can give some reference to some deleted assets. We don't want to migrate these.
 		const bool bAssetExist = AssetRegistry.GetAssetPackageDataCopy(*DependsIt).IsSet();
 
-		if ( !bIsEnginePackage && !bIsScriptPackage && bAssetExist)
+		if (!bIsScriptPackage && bAssetExist)
 		{
 			if (!AllDependencies.Contains(*DependsIt))
 			{
