Fixed an issue caused by the asset migration.
It was leaving some invalid asset data in the asset registry. These were from the temporally created assets during the migration.

diff --git a/Engine/Source/Developer/AssetTools/Private/AssetTools.cpp b/Engine/Source/Developer/AssetTools/Private/AssetTools.cpp
--- a/Engine/Source/Developer/AssetTools/Private/AssetTools.cpp
+++ b/Engine/Source/Developer/AssetTools/Private/AssetTools.cpp
@@ -974,6 +974,9 @@ namespace UE::AssetTools::Private
 
 		static void CleanInstancedPackages(const TArray<TWeakObjectPtr<UPackage>>& PackagesToClean)
 		{
+			FAssetRegistryModule& AssetRegistryModule = FModuleManager::LoadModuleChecked<FAssetRegistryModule>(TEXT("AssetRegistry"));
+			IAssetRegistry& AssetRegistry = AssetRegistryModule.Get();
+
 			TSharedRef<TFunction<void(UObject*)>> PurgeObject = MakeShared<TFunction<void(UObject*)>>();
 			PurgeObject.Get() = [PurgeObject](UObject* Object)
 			{
@@ -1002,18 +1005,26 @@ namespace UE::AssetTools::Private
 			FGlobalComponentReregisterContext ComponentContext;
 
 			TArray<UObject*> ReferenceToNull;
+
+			// We do the clean pass of the packages in two loop because the PurgeObject can affect the ability to get the main object from another package.
 			for (const TWeakObjectPtr<UPackage>& WeakPackage : PackagesToClean)
 			{
 				if (UPackage* Package = WeakPackage.Get())
 				{
-					const ERenameFlags PkgRenameFlags = REN_ForceNoResetLoaders | REN_DoNotDirty | REN_DontCreateRedirectors | REN_NonTransactional | REN_SkipGeneratedClasses;
-					check(Package->Rename(*MakeUniqueObjectName(nullptr, UPackage::StaticClass(), *FString::Printf(TEXT("%s_DEADFROMMIGRATION"), *Package->GetName())).ToString(), nullptr, PkgRenameFlags));
-
 					if (UObject* Asset = FPackageMigrationImpl::FindAssetInPackage(Package))
 					{
+						AssetRegistry.AssetDeleted(Asset);
 						ReferenceToNull.Add(Asset);
 					}
+				}
+			}
 
+			for (const TWeakObjectPtr<UPackage>& WeakPackage : PackagesToClean)
+			{
+				if (UPackage* Package = WeakPackage.Get())
+				{
+					const ERenameFlags PkgRenameFlags = REN_ForceNoResetLoaders | REN_DoNotDirty | REN_DontCreateRedirectors | REN_NonTransactional | REN_SkipGeneratedClasses;
+					check(Package->Rename(*MakeUniqueObjectName(nullptr, UPackage::StaticClass(), *FString::Printf(TEXT("%s_DEADFROMMIGRATION"), *Package->GetName())).ToString(), nullptr, PkgRenameFlags));
 					(*PurgeObject)(Package);
 					ForEachObjectWithOuter(Package, *PurgeObject);
 				}
