AssetTools: Fix allocation typo

diff --git a/Engine/Source/Developer/AssetTools/Private/PackageMigrationContext.cpp b/Engine/Source/Developer/AssetTools/Private/PackageMigrationContext.cpp
--- a/Engine/Source/Developer/AssetTools/Private/PackageMigrationContext.cpp
+++ b/Engine/Source/Developer/AssetTools/Private/PackageMigrationContext.cpp
@@ -118,7 +118,7 @@ FString FPackageMigrationContext::FScopedMountPoint::GetNewPackageNameForMigrati
 
 
 	FString NewPackageName;
-	NewPackageName.Reserve(GetRootPath().Len() - CurrentPackageWithoutMountPoint.Len());
+	NewPackageName.Reserve(GetRootPath().Len() + CurrentPackageWithoutMountPoint.Len());
 	NewPackageName.Append(GetRootPath());
 	NewPackageName.Append(CurrentPackageWithoutMountPoint);
 
