Fixed a crash when an asset save fail and report a error during a migration. Also fixed a crash when saving over a read only file during a migration.

diff --git a/Engine/Source/Developer/AssetTools/Private/AssetTools.cpp b/Engine/Source/Developer/AssetTools/Private/AssetTools.cpp
--- a/Engine/Source/Developer/AssetTools/Private/AssetTools.cpp
+++ b/Engine/Source/Developer/AssetTools/Private/AssetTools.cpp
@@ -514,6 +514,25 @@ namespace UE::AssetTools::Private
 			return Asset;
 		};
 
+		static bool CanSavePackageToFile(const FString& DestinationFile, const FStringView PackageName, FPackageMigrationContext* OptionalPackageMigrationContext)
+		{
+			// Don't try to migrate the package if the destination is read only
+			if (IFileManager::Get().IsReadOnly(*DestinationFile))
+			{
+				if (OptionalPackageMigrationContext)
+				{
+					OptionalPackageMigrationContext->AddErrorMigrationMessage(FText::Format(LOCTEXT("MigratePackages_SaveFailedReadOnly", "Couldn't migrate package ({0}) because the destination file is read only. Destination File ({1})")
+						, FText::FromStringView(PackageName)
+						, FText::FromString(DestinationFile)
+						));
+				}
+
+				return false;
+			}
+
+			return true;
+		}
+
 		// Prepare the MigrationPackagesData in the migration context and log any info about the package that will be renamed
 		static void SetupPublicAssetPackagesMigrationData(const TSharedPtr<TArray<ReportPackageData>>& PackageDataToMigrate
 			, TArray<TPair<const ReportPackageData*, const FAssetData>>& ExternalPackageDatas
@@ -618,30 +637,36 @@ namespace UE::AssetTools::Private
 				// Ask the user what to do if an asset already exist in the destination
 				if (bIsNameChangeCausedByExistingAssetInDestination)
 				{
-					// Handle name collision in the destination project. Post 5.1 should expose the possibility to rename the migrated package.
-
-					EAppReturnType::Type Response;
-					if (FApp::IsUnattended() || !PackageMigrationImplContext.Options.bPrompt || LastResponse == EAppReturnType::YesAll || LastResponse == EAppReturnType::NoAll)
+					if (!CanSavePackageToFile(NewPackageFilename, PackageData.Name, &PackageMigrationContext))
 					{
-						Response = LastResponse;
+						bShouldMigrate = false;
 					}
 					else
-					{
-						const FText Message = FText::Format(LOCTEXT("MigratePackages_AlreadyExists", "An asset already exists at location {0} would you like to overwrite it?"), FText::FromString(NewPackageFilename));
-						Response = FMessageDialog::Open(EAppMsgType::YesNoYesAllNoAllCancel, Message);
-						LastResponse = Response;
-					}
+					{ 
+						// Handle name collision in the destination project. Post 5.1 should expose the possibility to rename the migrated package
+						EAppReturnType::Type Response;
+						if (FApp::IsUnattended() || !PackageMigrationImplContext.Options.bPrompt || LastResponse == EAppReturnType::YesAll || LastResponse == EAppReturnType::NoAll)
+						{
+							Response = LastResponse;
+						}
+						else
+						{
+							const FText Message = FText::Format(LOCTEXT("MigratePackages_AlreadyExists", "An asset already exists at location {0} would you like to overwrite it?"), FText::FromString(NewPackageFilename));
+							Response = FMessageDialog::Open(EAppMsgType::YesNoYesAllNoAllCancel, Message);
+							LastResponse = Response;
+						}
 
-					if (Response == EAppReturnType::Cancel)
-					{
-						// The user chose to cancel mid-operation. Break out.
-						PackageMigrationImplContext.bWasCanceled = true;
-						break;
-					}
+						if (Response == EAppReturnType::Cancel)
+						{
+							// The user chose to cancel mid-operation. Break out.
+							PackageMigrationImplContext.bWasCanceled = true;
+							break;
+						}
 
-					if (Response == EAppReturnType::No || Response == EAppReturnType::NoAll)
-					{
-						bShouldMigrate = false;
+						if (Response == EAppReturnType::No || Response == EAppReturnType::NoAll)
+						{
+							bShouldMigrate = false;
+						}
 					}
 				}
 
@@ -902,7 +927,11 @@ namespace UE::AssetTools::Private
 		static void SaveInstancedPackagesIntoDestination(FPackageMigrationContext& PackageMigrationContext, FPackageMigrationImplContext& MigrationImplContext)
 		{
 			FSavePackageArgs SaveArgs;
-			SaveArgs.SaveFlags |= SAVE_RehydratePayloads | SAVE_Async;
+			SaveArgs.SaveFlags |= SAVE_RehydratePayloads;
+
+			// We should look into creating our own log to report the save erros to the user.
+			SaveArgs.Error = GWarn;
+
 
 			uint32 ProgressCount = 0;
 
