The import of assets via the asset tools now validate that destination path is mounted and avoid crashing if this is not the case.

diff --git a/Engine/Source/Developer/AssetTools/Private/AssetTools.cpp b/Engine/Source/Developer/AssetTools/Private/AssetTools.cpp
--- a/Engine/Source/Developer/AssetTools/Private/AssetTools.cpp
+++ b/Engine/Source/Developer/AssetTools/Private/AssetTools.cpp
@@ -3045,6 +3045,21 @@ TArray<UObject*> UAssetToolsImpl::ImportAssetsInternal(const TArray<FString>& Fi
 			continue;
 		}
 
+		// TryConvertFilenameToLongPackageName doesn't check if the package path is mounted but IsValidPath does it.
+		if (!FPackageName::IsValidPath(DestinationPath))
+		{
+			const FText Message = FText::Format(LOCTEXT("InvalidDestinationPath", "Can't import the file '{0}' because the destination path '{1}' is not a valid package path for the current project."), FText::FromString(Filename), FText::FromString(DestinationPath));
+			if (!bAutomatedImport)
+			{
+				FMessageDialog::Open(EAppMsgType::Ok, Message);
+			}
+
+			UE_LOG(LogAssetTools, Warning, TEXT("%s"), *ErrorMsg);
+			UE_LOG(LogAssetTools, Warning, TEXT("%s"), *Message.ToString());
+
+			continue;
+		}
+
 		if (bUseInterchangeFramework)
 		{
 			UE::Interchange::FScopedSourceData ScopedSourceData(Filename);
