The migration can now change the writing permissions temporally to save the migrated assets.
The migration also now report in its callback the dependencies that where not exported because of the new filtering of the assets that can be migrated.
Fixed some other issues like a crash and a bad event broadcast.

diff --git a/Engine/Source/Developer/AssetTools/Private/AssetTools.cpp b/Engine/Source/Developer/AssetTools/Private/AssetTools.cpp
--- a/Engine/Source/Developer/AssetTools/Private/AssetTools.cpp
+++ b/Engine/Source/Developer/AssetTools/Private/AssetTools.cpp
@@ -451,21 +451,69 @@ namespace UE::AssetTools::Private
 		};
 
 
-		static FString GetMountPointRootPath(bool bIsDestinationAProject, const FString& DestinationFolder)
+		struct FTemporaryWritableFolderPermission
+		{
+			FTemporaryWritableFolderPermission(const TSharedRef<FPathPermissionList>& InWritableFolder)
+				: WritableFolder(InWritableFolder)
+			{
+			}
+
+			~FTemporaryWritableFolderPermission()
+			{
+				for (const FName Owner : DenyList)
+				{
+					WritableFolder->AddDenyListItem(Owner, TempWritableRootPath);
+				}
+
+				WritableFolder->UnregisterOwner(TemporalyPermittedByMigration);
+			}
+
+			void AddTemporaryWriteFolderPermission(const FStringView InWritableRootPath)
+			{
+				TempWritableRootPath = InWritableRootPath;
+
+				if (const FPermissionListOwners* DenyListPtr = WritableFolder->GetDenyList().Find(TempWritableRootPath))
+				{
+					DenyList = *DenyListPtr;
+
+					for (const FName Owner : DenyList)
+					{
+						WritableFolder->RemoveDenyListItem(Owner, TempWritableRootPath);
+					}
+				}
+
+				// No need to add a permission if the list is empty since everything is permitted (except the deny list)
+				if (!WritableFolder->GetAllowList().IsEmpty())
+				{
+					WritableFolder->AddAllowListItem(TemporalyPermittedByMigration, TempWritableRootPath);
+				}
+			}
+
+		private:
+			TSharedRef<FPathPermissionList> WritableFolder;
+			FString TempWritableRootPath;
+			FPermissionListOwners DenyList;
+
+			static const inline TCHAR* TemporalyPermittedByMigration = TEXT("TemporalyPermittedByMigration");
+		};
+
+		static FString GetMountPointRootPath(const FString& DestinationFolder)
 		{
 			FString MountPointRootPath;
-			if (bIsDestinationAProject)
+
+			TArray<FString> Files;
+			IFileManager::Get().FindFiles(Files, *(DestinationFolder + TEXT("../")), TEXT("uproject"));
+			if (!Files.IsEmpty())
 			{
 				MountPointRootPath = TEXT("/Game/");
 			}
 			else
 			{
-				TArray<FString> PluginFiles;
-				IFileManager::Get().FindFiles(PluginFiles, *(DestinationFolder + TEXT("../")), TEXT("uplugin"));
+				IFileManager::Get().FindFiles(Files, *(DestinationFolder + TEXT("../")), TEXT("uplugin"));
 
-				if (PluginFiles.Num() == 1)
+				if (Files.Num() == 1)
 				{
-					FStringView PluginNameView = FPathViews::GetBaseFilename(PluginFiles[0]);
+					FStringView PluginNameView = FPathViews::GetBaseFilename(Files[0]);
 					MountPointRootPath.Reserve(PluginNameView.Len() + 2);
 					MountPointRootPath.AppendChar(TEXT('/'));
 					MountPointRootPath.Append(PluginNameView);
@@ -727,9 +775,9 @@ namespace UE::AssetTools::Private
 				ExistingPackageNameToMigrationDataIndex.Add(PackageData.Name, MigrationPackageDataIndex);
 
 				// Check if there is a existing package that need to be moved out the way
-				if (UPackage* InTheWayPacakge = FindObjectFast<UPackage>(nullptr, *NewPackageName))
+				if (UPackage* InTheWayPackage = FindObjectFast<UPackage>(nullptr, *NewPackageName))
 				{
-					PackageMigrationContext.MoveInTheWayPackage(InTheWayPacakge);
+					PackageMigrationContext.MoveInTheWayPackage(InTheWayPackage);
 				}
 
 				PackageMigrationImplContext.InstancingContext.AddPackageMapping(*PackageData.Name, *NewPackageName);
@@ -822,15 +870,57 @@ namespace UE::AssetTools::Private
 				MigrationData.bNeedToBeSaveMigrated = OuterMigrationData.bNeedToBeSaveMigrated && PackageData.Key->bShouldMigratePackage;
 
 				// Check if there is a existing package that need to be moved out the way
-				if (UPackage* InTheWayPacakge = FindObjectFast<UPackage>(nullptr, *NewPackageName))
+				if (UPackage* InTheWayPackage = FindObjectFast<UPackage>(nullptr, *NewPackageName))
 				{
-					PackageMigrationContext.MoveInTheWayPackage(InTheWayPacakge);
+					PackageMigrationContext.MoveInTheWayPackage(InTheWayPackage);
 				}
 
 				MigrationImplContext.InstancingContext.AddPackageMapping(*(PackageData.Key->Name), *NewPackageName);
 			}
 		}
 
+		static void ProcessExcludedDependencies(const TSet<FName>& ExcludedDependencies, UE::AssetTools::FPackageMigrationContext& PackageMigrationContext)
+		{
+			// Clean the excluded dependencies that are not used by the migrated packages
+			if (!ExcludedDependencies.IsEmpty())
+			{
+				TSet<FName> PackageToMigrate;
+
+				const TArray<UE::AssetTools::FPackageMigrationContext::FMigrationPackageData>& MigrationPackages = PackageMigrationContext.GetMigrationPackagesData();
+				PackageToMigrate.Reserve(MigrationPackages.Num());
+
+				for (const UE::AssetTools::FPackageMigrationContext::FMigrationPackageData& PackageData : MigrationPackages)
+				{
+					if (PackageData.bNeedToBeSaveMigrated)
+					{
+						PackageToMigrate.Emplace(*PackageData.OriginalPackageName);
+					}
+				}
+
+				PackageMigrationContext.ExcludedDependencies.Reserve(ExcludedDependencies.Num());
+
+				FAssetRegistryModule& AssetRegistryModule = FModuleManager::LoadModuleChecked<FAssetRegistryModule>(TEXT("AssetRegistry"));
+				IAssetRegistry& AssetRegistry = AssetRegistryModule.Get();
+				TArray<FName> Referencers;
+
+				for (const FName ExcludedDependency : ExcludedDependencies)
+				{
+					AssetRegistry.GetReferencers(ExcludedDependency, Referencers);
+
+					for (const FName Reference : Referencers)
+					{
+						if (PackageToMigrate.Contains(Reference))
+						{
+							PackageMigrationContext.ExcludedDependencies.Add(ExcludedDependency.ToString());
+							break;
+						}
+					}
+
+					Referencers.Reset();
+				}
+			}
+		}
+
 		static void CreatePackagesAndSetupLinkers(FPackageMigrationContext& PackageMigrationContext, FPackageMigrationImplContext& MigrationImplContext, TArray<TWeakObjectPtr<UPackage>>& PackagesToClean)
 		{
 			FArchive* ReaderOverride = nullptr;
@@ -917,7 +1007,7 @@ namespace UE::AssetTools::Private
 			FSavePackageArgs SaveArgs;
 			SaveArgs.SaveFlags |= SAVE_RehydratePayloads;
 
-			// We should look into creating our own log to report the save erros to the user.
+			// We should look into creating our own log to report the save errors to the user.
 			SaveArgs.Error = GWarn;
 
 
@@ -989,8 +1079,14 @@ namespace UE::AssetTools::Private
 			}
 		}
 
-		static void CleanInstancedPackages(const TArray<TWeakObjectPtr<UPackage>>& PackagesToClean)
+		static void CleanInstancedPackages(const TArray<TWeakObjectPtr<UPackage>>& PackagesToClean, FPackageMigrationContext& PackageMigrationContext)
 		{
+			for (FPackageMigrationContext::FMigrationPackageData& MigrationPackageData : PackageMigrationContext.MigrationPackagesData)
+			{
+				// Fixes some reference issues when the pending kill flag is off
+				MigrationPackageData.InstancedPackage = nullptr;
+			}
+
 			FAssetRegistryModule& AssetRegistryModule = FModuleManager::LoadModuleChecked<FAssetRegistryModule>(TEXT("AssetRegistry"));
 			IAssetRegistry& AssetRegistry = AssetRegistryModule.Get();
 
@@ -1085,7 +1181,7 @@ namespace UE::AssetTools::Private
 
 			if (!PackageMigrationContext.MigratedPackageMessages.IsEmpty())
 			{
-				MigrateLog.NewPage(LOCTEXT("MigratePackages_SuccessPageLabel", "Migrated Pacakges"));
+				MigrateLog.NewPage(LOCTEXT("MigratePackages_SuccessPageLabel", "Migrated Packages"));
 				for (const FText& InfoMessage : PackageMigrationContext.MigratedPackageMessages)
 				{
 					MigrateLog.Info(InfoMessage);
@@ -4038,10 +4134,13 @@ void UAssetToolsImpl::PerformMigratePackages(TArray<FName> PackageNamesToMigrate
 	// Form a full list of packages to move by including the dependencies of the supplied packages
 	TSet<FName> AllPackageNamesToMove;
 	TSet<FString> ExternalObjectsPaths;
+	TSet<FName> ExcludedDependencies;
 	{
 		FScopedSlowTask SlowTask( PackageNamesToMigrate.Num(), LOCTEXT( "MigratePackages_GatheringDependencies", "Gathering Dependencies..." ) );
 		SlowTask.MakeDialog();
 
+		TFunction<bool (FName)> ShouldExcludePackage = [this](FName PackageName) { return !CanMigratePackage(PackageName); };
+
 		for ( auto PackageIt = PackageNamesToMigrate.CreateConstIterator(); PackageIt; ++PackageIt )
 		{
 			SlowTask.EnterProgressFrame();
@@ -4049,7 +4148,7 @@ void UAssetToolsImpl::PerformMigratePackages(TArray<FName> PackageNamesToMigrate
 			if ( !AllPackageNamesToMove.Contains(*PackageIt) )
 			{
 				AllPackageNamesToMove.Add(*PackageIt);
-				RecursiveGetDependencies(*PackageIt, AllPackageNamesToMove, ExternalObjectsPaths);
+				RecursiveGetDependencies(*PackageIt, AllPackageNamesToMove, ExternalObjectsPaths, ExcludedDependencies, ShouldExcludePackage);
 			}
 		}
 	}
@@ -4147,16 +4246,16 @@ void UAssetToolsImpl::PerformMigratePackages(TArray<FName> PackageNamesToMigrate
 	if(!FApp::IsUnattended() && Options.bPrompt)
 	{
 		const FText ReportMessage = LOCTEXT("MigratePackagesReportTitle", "The following assets will be migrated to another content folder.");
-		SPackageReportDialog::FOnReportConfirmed OnReportConfirmed = SPackageReportDialog::FOnReportConfirmed::CreateUObject(this, &UAssetToolsImpl::MigratePackages_ReportConfirmed, ReportPackages, DestinationPath, Options);
+		SPackageReportDialog::FOnReportConfirmed OnReportConfirmed = SPackageReportDialog::FOnReportConfirmed::CreateUObject(this, &UAssetToolsImpl::MigratePackages_ReportConfirmed, ReportPackages, DestinationPath, MoveTemp(ExcludedDependencies), Options);
 		SPackageReportDialog::OpenPackageReportDialog(ReportMessage, *ReportPackages.Get(), OnReportConfirmed);
 	}
 	else
 	{
-		UAssetToolsImpl::MigratePackages_ReportConfirmed(ReportPackages, DestinationPath, Options);
+		UAssetToolsImpl::MigratePackages_ReportConfirmed(ReportPackages, DestinationPath, MoveTemp(ExcludedDependencies), Options);
 	}
 }
 
-void UAssetToolsImpl::MigratePackages_ReportConfirmed(TSharedPtr<TArray<ReportPackageData>> PackageDataToMigrate, const FString DestinationPath, const FMigrationOptions Options) const
+void UAssetToolsImpl::MigratePackages_ReportConfirmed(TSharedPtr<TArray<ReportPackageData>> PackageDataToMigrate, const FString DestinationPath, TSet<FName> ExcludedDependencies, const FMigrationOptions Options) const
 {
 	FString DestinationFolder;
 	if (FApp::IsUnattended() || !DestinationPath.IsEmpty())
@@ -4204,20 +4303,51 @@ void UAssetToolsImpl::MigratePackages_ReportConfirmed(TSharedPtr<TArray<ReportPa
 				}
 
 				// Verify that it is a content folder
-				if (DestinationFolder.EndsWith(TEXT("/Content/")))
+				if (!DestinationFolder.EndsWith(TEXT("/Content/")))
 				{
-					bFolderAccepted = true;
+					// The user chose a non-content folder. Let them know they cannot do that.
+					const FText Message = FText::Format(LOCTEXT("MigratePackages_NonContentFolder", "{0} does not appear to be a Content folder. Migrated content only work if placed in a Content folder. Select a Content folder."), FText::FromString(DestinationFolder));
+					EAppReturnType::Type Response = FMessageDialog::Open(EAppMsgType::OkCancel, Message);
+					if (Response == EAppReturnType::Cancel)
+					{
+						return;
+					}
+
+					continue;
 				}
-				else
+
+
+				FString RootPath = UE::AssetTools::Private::FPackageMigrationImpl::GetMountPointRootPath(DestinationFolder);
+
+				if (RootPath.IsEmpty())
 				{
-					// The user chose a non-content folder. Let them know they cannot do that.
-					const FText Message = FText::Format(LOCTEXT("MigratePackages_NonContentFolder", "{0} does not appear to be a game Content folder. Migrated content only work if placed in a Content folder. Select a Content folder."), FText::FromString(DestinationFolder));
+					const FText Message = FText::Format(LOCTEXT("MigratePackages_CannotIdentifyMountPoint", "{0} does not appear to be a game Content folder or a Plugin folder. The tool didn't found the associated uproject or uplugin file used to determine the unreal path of destination folder. Migrated content only work if placed in a Project or a Plugin Content folder. Select a Content folder."), FText::FromString(DestinationFolder));
 					EAppReturnType::Type Response = FMessageDialog::Open(EAppMsgType::OkCancel, Message);
 					if (Response == EAppReturnType::Cancel)
 					{
 						return;
 					}
+					continue;
+				}
+
+				if (WritableFolderPermissionList->HasFiltering() && FPathViews::IsParentPathOf(FPaths::ProjectDir(), DestinationFolder))
+				{
+					// Make sure the path is not in the deny list
+					if (!WritableFolderPermissionList->PassesFilter(RootPath))
+					{
+						const FText Message = FText::Format(LOCTEXT("MigratePackages_CannotMoveContentHere", "{0} is not in a writable folder. Chose another Content folder."), FText::FromString(DestinationFolder));
+						EAppReturnType::Type Response = FMessageDialog::Open(EAppMsgType::OkCancel, Message);
+						if (Response == EAppReturnType::Cancel)
+						{
+							return;
+						}
+						continue;
+					}
 				}
+				
+			
+
+				bFolderAccepted = true;
 			}
 		}
 		else
@@ -4240,19 +4370,32 @@ void UAssetToolsImpl::MigratePackages_ReportConfirmed(TSharedPtr<TArray<ReportPa
 
 	FMessageLog MigrateLog("AssetTools");
 
-	// Determine if the destination is a project content folder or a plugin content folder
-	TArray<FString> ProjectFiles;
-	IFileManager::Get().FindFiles(ProjectFiles, *(DestinationFolder + TEXT("../")), TEXT("uproject"));
-	bool bIsDestinationAProject = !ProjectFiles.IsEmpty();
-
 	if (UE::AssetTools::Private::bUseNewPackageMigration)
 	{
 		using namespace UE::AssetTools::Private;
 
-		FString MountPointRootPath = FPackageMigrationImpl::GetMountPointRootPath(bIsDestinationAProject, DestinationFolder);
+		FString MountPointRootPath = FPackageMigrationImpl::GetMountPointRootPath(DestinationFolder);
 
 		if (!MountPointRootPath.IsEmpty())
 		{
+			FPackageMigrationImpl::FTemporaryWritableFolderPermission ScopedFolderPermission(WritableFolderPermissionList);
+
+			// Add a temporary permission if needed
+			if (WritableFolderPermissionList->HasFiltering())
+			{
+				if (FPathViews::IsParentPathOf(FPaths::ProjectDir(), DestinationFolder))
+				{
+					if (!WritableFolderPermissionList->PassesFilter(MountPointRootPath))
+					{
+						UE_LOG(LogAssetTools, Error, TEXT("Migration Destination path is not a writable folder."));
+						return;
+					}
+				}
+				else
+				{
+					ScopedFolderPermission.AddTemporaryWriteFolderPermission(MountPointRootPath);
+				}
+			}
 
 			FPackageMigrationImpl::FPackageMigrationImplContext MigrationImplContext;
 			MigrationImplContext.Options = Options;
@@ -4289,6 +4432,11 @@ void UAssetToolsImpl::MigratePackages_ReportConfirmed(TSharedPtr<TArray<ReportPa
 				PackageMigrationContext.CurrentStep = UE::AssetTools::FPackageMigrationContext::EPackageMigrationStep::PostExternalMigrationPackageDataCreated;
 				OnPackageMigration.Broadcast(PackageMigrationContext);
 
+
+				FPackageMigrationImpl::ProcessExcludedDependencies(ExcludedDependencies, PackageMigrationContext);
+				
+				PackageMigrationContext.CurrentStep = UE::AssetTools::FPackageMigrationContext::EPackageMigrationStep::PostExcludedDependenciesCreated;
+				OnPackageMigration.Broadcast(PackageMigrationContext);
 			}
 
 			PackageMigrationContext.CurrentStep = UE::AssetTools::FPackageMigrationContext::EPackageMigrationStep::InTheWayPackagesMoved;
@@ -4320,14 +4468,15 @@ void UAssetToolsImpl::MigratePackages_ReportConfirmed(TSharedPtr<TArray<ReportPa
 
 
 			// 5) Remove the packages created for the migration
-			FPackageMigrationImpl::CleanInstancedPackages(PackagesToClean);
+			FPackageMigrationImpl::CleanInstancedPackages(PackagesToClean, PackageMigrationContext);
 
 			PackageMigrationContext.CurrentStep = UE::AssetTools::FPackageMigrationContext::EPackageMigrationStep::PostCleaningInstancedPackages;
 			OnPackageMigration.Broadcast(PackageMigrationContext);
 
+			// 6) Restore the existing package that where in the way
 			FPackageMigrationImpl::RestoreInTheWayPackages(PackageMigrationContext);
 
-			PackageMigrationContext.CurrentStep = UE::AssetTools::FPackageMigrationContext::EPackageMigrationStep::InstancedPackagesSaved;
+			PackageMigrationContext.CurrentStep = UE::AssetTools::FPackageMigrationContext::EPackageMigrationStep::InTheWayPackagesRestored;
 			OnPackageMigration.Broadcast(PackageMigrationContext);
 
 			FPackageMigrationImpl::BuildAndNotifyLogWindow(PackageMigrationContext, MigrateLog, MigrationImplContext);
@@ -4346,6 +4495,11 @@ void UAssetToolsImpl::MigratePackages_ReportConfirmed(TSharedPtr<TArray<ReportPa
 		// Assets in /Game always map directly to the destination
 		DestContentRootsToFolders.Add(FName(TEXT("/Game")), DestinationFolder);
 
+		// Determine if the destination is a project content folder
+		TArray<FString> ProjectFiles;
+		IFileManager::Get().FindFiles(ProjectFiles, *(DestinationFolder + TEXT("../")), TEXT("uproject"));
+		bool bIsDestinationAProject = !ProjectFiles.IsEmpty();
+
 		// If our destination is a project, it could have plugins...
 		if (bIsDestinationAProject)
 		{
@@ -4914,14 +5068,30 @@ void UAssetToolsImpl::MigratePackages_ReportConfirmed(TSharedPtr<TArray<ReportPa
 	}
 }
 
-void UAssetToolsImpl::RecursiveGetDependencies(const FName& PackageName, TSet<FName>& AllDependencies, TSet<FString>& OutExternalObjectsPaths) const
+bool UAssetToolsImpl::CanMigratePackage(FName PackageName) const
+{
+	for (const TPair<FName, UE::AssetTools::FCanMigrateAsset>& Pair : CanMigrateAssetDelegates)
+	{
+		if (Pair.Value.IsBound())
+		{
+			if (!Pair.Value.Execute(PackageName))
+			{
+				return false;
+			}
+		}
+	}
+
+	return true;
+}
+
+void UAssetToolsImpl::RecursiveGetDependencies(const FName& PackageName, TSet<FName>& AllDependencies, TSet<FString>& OutExternalObjectsPaths, TSet<FName>& ExcludedDependencies, const TFunction<bool (FName)>& ShouldExcludeFromDependenciesSearch) const
 {
 	FAssetRegistryModule& AssetRegistryModule = FModuleManager::Get().LoadModuleChecked<FAssetRegistryModule>(TEXT("AssetRegistry"));
 	TArray<FName> Dependencies;
 	IAssetRegistry& AssetRegistry = AssetRegistryModule.Get();
 	AssetRegistry.GetDependencies(PackageName, Dependencies);
 	
-	for ( auto DependsIt = Dependencies.CreateConstIterator(); DependsIt; ++DependsIt )
+	for (TArray<FName>::TConstIterator DependsIt = Dependencies.CreateConstIterator(); DependsIt; ++DependsIt)
 	{
 		FString DependencyName = (*DependsIt).ToString();
 
@@ -4932,10 +5102,19 @@ void UAssetToolsImpl::RecursiveGetDependencies(const FName& PackageName, TSet<FN
 
 		if (!bIsScriptPackage && bAssetExist)
 		{
-			if (!AllDependencies.Contains(*DependsIt))
+			uint32 DependsHash = GetTypeHash(*DependsIt);
+			if (!AllDependencies.ContainsByHash(DependsHash, *DependsIt) && !ExcludedDependencies.ContainsByHash(DependsHash, *DependsIt))
 			{
-				AllDependencies.Add(*DependsIt);
-				RecursiveGetDependencies(*DependsIt, AllDependencies, OutExternalObjectsPaths);
+				// Early stop the dependency search
+				if (ShouldExcludeFromDependenciesSearch(*DependsIt))
+				{
+					ExcludedDependencies.AddByHash(DependsHash, *DependsIt);
+					continue;
+				}
+
+				AllDependencies.AddByHash(DependsHash, *DependsIt);
+
+				RecursiveGetDependencies(*DependsIt, AllDependencies, OutExternalObjectsPaths, ExcludedDependencies, ShouldExcludeFromDependenciesSearch);
 			}
 		}
 	}
@@ -4965,8 +5144,9 @@ void UAssetToolsImpl::RecursiveGetDependencies(const FName& PackageName, TSet<FN
 
 						for (const FAssetData& ExternalObjectAsset : ExternalObjectAssets)
 						{
+							// We don't expose the early dependency search exit to the external objects/actors since to the users their are same the outer package that own these objects
 							AllDependencies.Add(ExternalObjectAsset.PackageName);
-							RecursiveGetDependencies(ExternalObjectAsset.PackageName, AllDependencies, OutExternalObjectsPaths);
+							RecursiveGetDependencies(ExternalObjectAsset.PackageName, AllDependencies, OutExternalObjectsPaths, ExcludedDependencies, ShouldExcludeFromDependenciesSearch);
 						}
 					}
 				}
@@ -5545,5 +5725,15 @@ void UAssetToolsImpl::UnregisterIsNameAllowedDelegate(const FName OwnerName)
 	IsNameAllowedDelegates.Remove(OwnerName);
 }
 
+void UAssetToolsImpl::RegisterCanMigrateAsset(const FName OwnerName, UE::AssetTools::FCanMigrateAsset Delegate)
+{
+	CanMigrateAssetDelegates.Add(OwnerName, MoveTemp(Delegate));
+}
+
+void UAssetToolsImpl::UnregisterCanMigrateAsset(const FName OwnerName)
+{
+	CanMigrateAssetDelegates.Remove(OwnerName);
+}
+
 #undef LOCTEXT_NAMESPACE
 
diff --git a/Engine/Source/Developer/AssetTools/Private/AssetTools.h b/Engine/Source/Developer/AssetTools/Private/AssetTools.h
--- a/Engine/Source/Developer/AssetTools/Private/AssetTools.h
+++ b/Engine/Source/Developer/AssetTools/Private/AssetTools.h
@@ -143,6 +143,9 @@ public:
 	virtual void RegisterIsNameAllowedDelegate(const FName OwnerName, FIsNameAllowed Delegate) override;
 	virtual void UnregisterIsNameAllowedDelegate(const FName OwnerName) override;
 	
+	virtual void RegisterCanMigrateAsset(const FName OwnerName, UE::AssetTools::FCanMigrateAsset Delegate) override;
+	virtual void UnregisterCanMigrateAsset(const FName OwnerName) override;
+
 	virtual void SyncBrowserToAssets(const TArray<UObject*>& AssetsToSync) override;
 	virtual void SyncBrowserToAssets(const TArray<FAssetData>& AssetsToSync) override;
 public:
@@ -166,19 +169,22 @@ private:
 	/** Begins the package migration, after assets have been discovered */
 	void PerformMigratePackages(TArray<FName> PackageNamesToMigrate, const FString DestinationPath, const FMigrationOptions Options) const;
 
+	/** Check if an asset can migrated */
+	bool CanMigratePackage(FName PackageName) const;
+
 	TArray<FName> ExpandAssetsAndFoldersToJustAssets(TArray<FName> SelectedAssetAndFolderNames) const;
 
 	/** Begins the package advanced copy, after assets have been discovered */
 	void PerformAdvancedCopyPackages(TArray<FName> SelectedPackageNames, FString TargetPath) const;
 
 	/** Copies files after the final list was confirmed */
-	void MigratePackages_ReportConfirmed(TSharedPtr<TArray<ReportPackageData>> PackageDataToMigrate, const FString DestinationPath, const FMigrationOptions Options) const;
+	void MigratePackages_ReportConfirmed(TSharedPtr<TArray<ReportPackageData>> PackageDataToMigrate, const FString DestinationPath, TSet<FName> ExcludedDependencies, const FMigrationOptions Options) const;
 
 	/** Copies files after the final list was confirmed */
 	void AdvancedCopyPackages_ReportConfirmed(FAdvancedCopyParams CopyParam, TArray<TMap<FString, FString>> DestinationMap) const;
 
 	/** Gets the dependencies of the specified package recursively */
-	void RecursiveGetDependencies(const FName& PackageName, TSet<FName>& AllDependencies, TSet<FString>& ExternalObjectsPaths) const;
+	void RecursiveGetDependencies(const FName& PackageName, TSet<FName>& AllDependencies, TSet<FString>& ExternalObjectsPaths, TSet<FName>& ExcludedDependencies, const TFunction<bool(FName)>& ShouldExcludeFromDependenciesSearch) const;
 
 	/** Gets the dependencies of the specified package recursively while omitting things that don't pass the FARFilter passed in from FAdvancedCopyParams */
 	void RecursiveGetDependenciesAdvanced(const FName& PackageName, FAdvancedCopyParams& CopyParams, TArray<FName>& AllDependencies, TMap<FName, FName>& DependencyMap, const class UAdvancedCopyCustomization* CopyCustomization, TArray<FAssetData>& OptionalAssetData) const;
@@ -248,6 +254,8 @@ private:
 	TMap<FName, FIsNameAllowed> IsNameAllowedDelegates;
 
 	UE::AssetTools::FOnPackageMigration OnPackageMigration;
+
+	TMap<FName, UE::AssetTools::FCanMigrateAsset> CanMigrateAssetDelegates;
 };
 
 PRAGMA_ENABLE_DEPRECATION_WARNINGS
diff --git a/Engine/Source/Developer/AssetTools/Private/PackageMigrationContext.cpp b/Engine/Source/Developer/AssetTools/Private/PackageMigrationContext.cpp
--- a/Engine/Source/Developer/AssetTools/Private/PackageMigrationContext.cpp
+++ b/Engine/Source/Developer/AssetTools/Private/PackageMigrationContext.cpp
@@ -50,6 +50,13 @@ const TArray<UPackage*>& FPackageMigrationContext::GetMovedOutOfTheWayPackages()
 	return PackagesThatWhereMoved;
 }
 
+const TArray<FString>& FPackageMigrationContext::GetExcludedDependencies() const
+{
+	checkf(CurrentStep < EPackageMigrationStep::PostExcludedDependenciesCreated, TEXT("The migration can't give the excluded packages before the where processed"));
+
+	return ExcludedDependencies;
+}
+
 void FPackageMigrationContext::MoveInTheWayPackage(UPackage* Package)
 {
 	check(Package);
diff --git a/Engine/Source/Developer/AssetTools/Public/IAssetTools.h b/Engine/Source/Developer/AssetTools/Public/IAssetTools.h
--- a/Engine/Source/Developer/AssetTools/Public/IAssetTools.h
+++ b/Engine/Source/Developer/AssetTools/Public/IAssetTools.h
@@ -33,6 +33,8 @@ namespace UE::AssetTools
 	struct FPackageMigrationContext;
 
 	DECLARE_MULTICAST_DELEGATE_OneParam(FOnPackageMigration, FPackageMigrationContext&);
+
+	DECLARE_DELEGATE_RetVal_OneParam(bool, FCanMigrateAsset, FName);
 }
 
 UENUM()
@@ -573,6 +575,10 @@ public:
 	/** Show notification that writable folder filter blocked an action */
 	virtual void NotifyBlockedByWritableFolderFilter() const = 0;
 
+	/** Allow to add some restrictions to the assets that can be migrated */
+	virtual void RegisterCanMigrateAsset(const FName OwnerName, UE::AssetTools::FCanMigrateAsset Delegate) = 0;
+	virtual void UnregisterCanMigrateAsset(const FName OwnerName) = 0;
+
 	/** Syncs the primary content browser to the specified assets, whether or not it is locked. Most syncs that come from AssetTools -feel- like they came from the content browser, so this is okay. */
 	virtual void SyncBrowserToAssets(const TArray<UObject*>& AssetsToSync) = 0;
 	virtual void SyncBrowserToAssets(const TArray<FAssetData>& AssetsToSync) = 0;
diff --git a/Engine/Source/Developer/AssetTools/Public/PackageMigrationContext.h b/Engine/Source/Developer/AssetTools/Public/PackageMigrationContext.h
--- a/Engine/Source/Developer/AssetTools/Public/PackageMigrationContext.h
+++ b/Engine/Source/Developer/AssetTools/Public/PackageMigrationContext.h
@@ -42,6 +42,9 @@ struct ASSETTOOLS_API FPackageMigrationContext : public FGCObject
 	// Return the array of non instanced package where moved because they were in the way of the migration
 	const TArray<UPackage*>& GetMovedOutOfTheWayPackages() const;
 
+	// Return the array of package dependencies that where not migrated because there where excluded by some automated filters
+	const TArray<FString>& GetExcludedDependencies() const;
+
 	enum class EPackageMigrationStep
 	{
 		// Called before the migration begin
@@ -50,6 +53,8 @@ struct ASSETTOOLS_API FPackageMigrationContext : public FGCObject
 		PostAssetMigrationPackageDataCreated,
 		// Called after the non external packages to migrate where processed
 		PostExternalMigrationPackageDataCreated,
+		// Called after the excluded depencies where processed
+		PostExcludedDependenciesCreated,
 		// Called after all the package that might be in the way for the migration where move
 		InTheWayPackagesMoved,
 		// Called after the instanced package where created with their load linker setup being already set to load the right file
@@ -60,6 +65,8 @@ struct ASSETTOOLS_API FPackageMigrationContext : public FGCObject
 		InstancedPackagesSaved,
 		// Called after the instanced package were removed
 		PostCleaningInstancedPackages,
+		// Called after all the package that were in the way for the migration where restored
+		InTheWayPackagesRestored,
 		// Called at the end of the migration. After the in the way package where restored but before the log is processed
 		EndMigration
 	};
@@ -108,7 +115,7 @@ struct ASSETTOOLS_API FPackageMigrationContext : public FGCObject
 		// Tell the migration that package must be loaded
 		bool bNeedInstancedLoad = false;
 
-		// Tell the migration that the package must migrated by save
+		// Tell the migration that the package must be migrated by save
 		bool bNeedToBeSaveMigrated = false;
 
 		FMigrationPackageData(const FString& InInstancedPackageName, const FString& InOriginalPackageName, const FString& InDestinationFilename);
@@ -180,6 +187,9 @@ private:
 
 	// The data associated to the packages that take part of the migration process.
 	TArray<FMigrationPackageData> MigrationPackagesData;
+
+	// The excluded dependencies of the asset that are being migrated (these don't include the asset a user might have exclude voluntary from the migration)
+	TArray<FString> ExcludedDependencies;
 };
 
 }
\ No newline at end of file
