Compile AR filter only once when getting Advanced Copy sources/dependencies

- Using Advanced Copy on assets with lots of dependencies spent most of the time just compiling the AR filter
- I kept the original IAssetRegistry::UseFilterToExcludeAssets() (using an FARFilter) as it is blueprint callable

diff --git a/Engine/Source/Developer/AssetTools/Private/AssetTools.cpp b/Engine/Source/Developer/AssetTools/Private/AssetTools.cpp
--- a/Engine/Source/Developer/AssetTools/Private/AssetTools.cpp
+++ b/Engine/Source/Developer/AssetTools/Private/AssetTools.cpp
@@ -2062,8 +2062,9 @@ void UAssetToolsImpl::GetAllAdvancedCopySources(FName SelectedPackage, FAdvanced
 		AssetRegistry.GetAssetsByPackageName(SelectedPackage, SourceAssetData);
 		// Check if this is a folder before using the filter to exclude assets
 		bool bIsFolder = SourceAssetData.Num() == 0;
-		FARFilter ExclusionFilter = CopyCustomization->GetARFilter();
-		AssetRegistry.UseFilterToExcludeAssets(SourceAssetData, ExclusionFilter);
+		FARCompiledFilter CompiledExclusionFilter;
+		AssetRegistry.CompileFilter(CopyCustomization->GetARFilter(), CompiledExclusionFilter);
+		AssetRegistry.UseFilterToExcludeAssets(SourceAssetData, CompiledExclusionFilter);
 		// If this is a valid asset
 		if (SourceAssetData.Num() > 0 || bIsFolder)
 		{
@@ -2074,7 +2075,7 @@ void UAssetToolsImpl::GetAllAdvancedCopySources(FName SelectedPackage, FAdvanced
 		// Folders should ALWAYS get checked for assets and subfolders
 		if ((CopyParams.bShouldCheckForDependencies && SourceAssetData.Num() > 0) || bIsFolder)
 		{
-			RecursiveGetDependenciesAdvanced(SelectedPackage, CopyParams, CurrentDependencies, DependencyMap, CopyCustomization, SourceAssetData);
+			RecursiveGetDependenciesAdvanced(SelectedPackage, CopyParams, CurrentDependencies, DependencyMap, CompiledExclusionFilter, SourceAssetData);
 		}
 		OutPackageNamesToCopy.Append(CurrentDependencies);
 	}
@@ -5155,7 +5156,7 @@ void UAssetToolsImpl::RecursiveGetDependencies(const FName& PackageName, TSet<FN
 	}
 }
 
-void UAssetToolsImpl::RecursiveGetDependenciesAdvanced(const FName& PackageName, FAdvancedCopyParams& CopyParams, TArray<FName>& AllDependencies, TMap<FName, FName>& DependencyMap, const UAdvancedCopyCustomization* CopyCustomization, TArray<FAssetData>& OptionalAssetData) const
+void UAssetToolsImpl::RecursiveGetDependenciesAdvanced(const FName& PackageName, FAdvancedCopyParams& CopyParams, TArray<FName>& AllDependencies, TMap<FName, FName>& DependencyMap, const FARCompiledFilter& CompiledExclusionFilter, TArray<FAssetData>& OptionalAssetData) const
 {
 	FAssetRegistryModule& AssetRegistryModule = FModuleManager::Get().LoadModuleChecked<FAssetRegistryModule>(TEXT("AssetRegistry"));
 	TArray<FName> Dependencies;
@@ -5170,13 +5171,12 @@ void UAssetToolsImpl::RecursiveGetDependenciesAdvanced(const FName& PackageName,
 			{
 				TArray<FAssetData> DependencyAssetData;
 				AssetRegistry.GetAssetsByPackageName(Dep, DependencyAssetData, true);
-				FARFilter ExclusionFilter = CopyCustomization->GetARFilter();
-				AssetRegistry.UseFilterToExcludeAssets(DependencyAssetData, ExclusionFilter);
+				AssetRegistry.UseFilterToExcludeAssets(DependencyAssetData, CompiledExclusionFilter);
 				if (DependencyAssetData.Num() > 0)
 				{
 					AllDependencies.Add(Dep);
 					DependencyMap.Add(Dep, PackageName);
-					RecursiveGetDependenciesAdvanced(Dep, CopyParams, AllDependencies, DependencyMap, CopyCustomization, DependencyAssetData);
+					RecursiveGetDependenciesAdvanced(Dep, CopyParams, AllDependencies, DependencyMap, CompiledExclusionFilter, DependencyAssetData);
 				}
 			}
 		}
@@ -5188,15 +5188,14 @@ void UAssetToolsImpl::RecursiveGetDependenciesAdvanced(const FName& PackageName,
 		if (AssetRegistry.HasAssets(PackageName) && AssetRegistry.GetAssetsByPath(PackageName, PathAssetData))
 		{
 
-			FARFilter ExclusionFilter = CopyCustomization->GetARFilter();
-			AssetRegistry.UseFilterToExcludeAssets(PathAssetData, ExclusionFilter);
+			AssetRegistry.UseFilterToExcludeAssets(PathAssetData, CompiledExclusionFilter);
 			for(const FAssetData& Asset : PathAssetData)
 			{
 				AllDependencies.Add(*Asset.GetPackage()->GetName());
 				// If we should check the assets we found for dependencies
 				if (CopyParams.bShouldCheckForDependencies)
 				{
-					RecursiveGetDependenciesAdvanced(FName(*Asset.GetPackage()->GetName()), CopyParams, AllDependencies, DependencyMap, CopyCustomization, PathAssetData);
+					RecursiveGetDependenciesAdvanced(FName(*Asset.GetPackage()->GetName()), CopyParams, AllDependencies, DependencyMap, CompiledExclusionFilter, PathAssetData);
 				}
 			}
 		}
@@ -5208,7 +5207,7 @@ void UAssetToolsImpl::RecursiveGetDependenciesAdvanced(const FName& PackageName,
 			for (const FString& SubPath : SubPaths)
 			{
 				TArray<FAssetData> EmptyArray;
-				RecursiveGetDependenciesAdvanced(FName(*SubPath), CopyParams, AllDependencies, DependencyMap, CopyCustomization, EmptyArray);
+				RecursiveGetDependenciesAdvanced(FName(*SubPath), CopyParams, AllDependencies, DependencyMap, CompiledExclusionFilter, EmptyArray);
 			}
 
 		}
diff --git a/Engine/Source/Developer/AssetTools/Private/AssetTools.h b/Engine/Source/Developer/AssetTools/Private/AssetTools.h
--- a/Engine/Source/Developer/AssetTools/Private/AssetTools.h
+++ b/Engine/Source/Developer/AssetTools/Private/AssetTools.h
@@ -187,7 +187,7 @@ private:
 	void RecursiveGetDependencies(const FName& PackageName, TSet<FName>& AllDependencies, TSet<FString>& ExternalObjectsPaths, TSet<FName>& ExcludedDependencies, const TFunction<bool(FName)>& ShouldExcludeFromDependenciesSearch) const;
 
 	/** Gets the dependencies of the specified package recursively while omitting things that don't pass the FARFilter passed in from FAdvancedCopyParams */
-	void RecursiveGetDependenciesAdvanced(const FName& PackageName, FAdvancedCopyParams& CopyParams, TArray<FName>& AllDependencies, TMap<FName, FName>& DependencyMap, const class UAdvancedCopyCustomization* CopyCustomization, TArray<FAssetData>& OptionalAssetData) const;
+	void RecursiveGetDependenciesAdvanced(const FName& PackageName, FAdvancedCopyParams& CopyParams, TArray<FName>& AllDependencies, TMap<FName, FName>& DependencyMap, const FARCompiledFilter& CompiledExclusionFilter, TArray<FAssetData>& OptionalAssetData) const;
 
 	/** Records the time taken for an import and reports it to engine analytics, if available */
 	static void OnNewImportRecord(UClass* AssetType, const FString& FileExtension, bool bSucceeded, bool bWasCancelled, const FDateTime& StartTime);
diff --git a/Engine/Source/Runtime/AssetRegistry/Private/AssetRegistry.cpp b/Engine/Source/Runtime/AssetRegistry/Private/AssetRegistry.cpp
--- a/Engine/Source/Runtime/AssetRegistry/Private/AssetRegistry.cpp
+++ b/Engine/Source/Runtime/AssetRegistry/Private/AssetRegistry.cpp
@@ -2970,6 +2970,11 @@ void UAssetRegistryImpl::UseFilterToExcludeAssets(TArray<FAssetData>& AssetDataL
 	}
 	FARCompiledFilter CompiledFilter;
 	CompileFilter(Filter, CompiledFilter);
+	UseFilterToExcludeAssets(AssetDataList, CompiledFilter);
+}
+
+void UAssetRegistryImpl::UseFilterToExcludeAssets(TArray<FAssetData>& AssetDataList, const FARCompiledFilter& CompiledFilter) const
+{
 	UE::AssetRegistry::Utils::RunAssetsThroughFilter(AssetDataList, CompiledFilter, UE::AssetRegistry::Utils::EFilterMode::Exclusive);
 }
 
diff --git a/Engine/Source/Runtime/AssetRegistry/Private/AssetRegistry.h b/Engine/Source/Runtime/AssetRegistry/Private/AssetRegistry.h
--- a/Engine/Source/Runtime/AssetRegistry/Private/AssetRegistry.h
+++ b/Engine/Source/Runtime/AssetRegistry/Private/AssetRegistry.h
@@ -87,6 +87,7 @@ public:
 	virtual void EnumerateSubPaths(const FName InBasePath, TFunctionRef<bool(FName)> Callback, bool bInRecurse) const override;
 	virtual void RunAssetsThroughFilter (TArray<FAssetData>& AssetDataList, const FARFilter& Filter) const override;
 	virtual void UseFilterToExcludeAssets(TArray<FAssetData>& AssetDataList, const FARFilter& Filter) const override;
+	virtual void UseFilterToExcludeAssets(TArray<FAssetData>& AssetDataList, const FARCompiledFilter& CompiledFilter) const override;
 	virtual bool IsAssetIncludedByFilter(const FAssetData& AssetData, const FARCompiledFilter& Filter) const override;
 	virtual bool IsAssetExcludedByFilter(const FAssetData& AssetData, const FARCompiledFilter& Filter) const override;
 	virtual void ExpandRecursiveFilter(const FARFilter& InFilter, FARFilter& ExpandedFilter) const override;
diff --git a/Engine/Source/Runtime/AssetRegistry/Public/AssetRegistry/IAssetRegistry.h b/Engine/Source/Runtime/AssetRegistry/Public/AssetRegistry/IAssetRegistry.h
--- a/Engine/Source/Runtime/AssetRegistry/Public/AssetRegistry/IAssetRegistry.h
+++ b/Engine/Source/Runtime/AssetRegistry/Public/AssetRegistry/IAssetRegistry.h
@@ -435,6 +435,9 @@ public:
 	UFUNCTION(BlueprintCallable, BlueprintPure=false, Category = "AssetRegistry")
 	virtual void UseFilterToExcludeAssets(UPARAM(ref) TArray<FAssetData>& AssetDataList, const FARFilter& Filter) const = 0;
 
+	/** Trims items out of the asset data list that pass the supplied filter */
+	virtual void UseFilterToExcludeAssets(TArray<FAssetData>& AssetDataList, const FARCompiledFilter& CompiledFilter) const = 0;
+	
 	/** Tests to see whether the given asset would be included (passes) the given filter */
 	virtual bool IsAssetIncludedByFilter(const FAssetData& AssetData, const FARCompiledFilter& Filter) const = 0;
 
