Added the auto save flag to the package migration to avoid some asset from generating stuff on save.

diff --git a/Engine/Source/Developer/AssetTools/Private/AssetTools.cpp b/Engine/Source/Developer/AssetTools/Private/AssetTools.cpp
--- a/Engine/Source/Developer/AssetTools/Private/AssetTools.cpp
+++ b/Engine/Source/Developer/AssetTools/Private/AssetTools.cpp
@@ -1005,7 +1005,9 @@ namespace UE::AssetTools::Private
 		static void SaveInstancedPackagesIntoDestination(FPackageMigrationContext& PackageMigrationContext, FPackageMigrationImplContext& MigrationImplContext)
 		{
 			FSavePackageArgs SaveArgs;
-			SaveArgs.SaveFlags |= SAVE_RehydratePayloads;
+
+			// Rehydrate the payloads so that we move the virtual payloads also. Add the auto save flag to notify the editor extensions that they shouldn't prompt the user or generate/update some asset on save.
+			SaveArgs.SaveFlags |= SAVE_RehydratePayloads | SAVE_FromAutosave;
 
 			// We should look into creating our own log to report the save errors to the user.
 			SaveArgs.Error = GWarn;
