If we import only one main asset (static mesh, skeletal mesh or anim sequence), we should select it in the content browser when the import is done.

diff --git a/Engine/Source/Developer/AssetTools/Private/AssetTools.cpp b/Engine/Source/Developer/AssetTools/Private/AssetTools.cpp
--- a/Engine/Source/Developer/AssetTools/Private/AssetTools.cpp
+++ b/Engine/Source/Developer/AssetTools/Private/AssetTools.cpp
@@ -97,6 +97,9 @@
 #include "InterchangeProjectSettings.h"
 #include "Engine/World.h"
 #include "Engine/Level.h"
+#include "Engine/SkeletalMesh.h"
+#include "Engine/StaticMesh.h"
+#include "Animation/AnimSequence.h"
 #include "UObject/SavePackage.h"
 #include "Dialogs/Dialogs.h"
 #include "Widgets/Input/SEditableTextBox.h"
@@ -3134,6 +3137,8 @@ TArray<UObject*> UAssetToolsImpl::ImportAssetsInternal(const TArray<FString>& Fi
 
 	TSharedPtr<FInterchangeImportStatus, ESPMode::ThreadSafe> ImportStatus = MakeShared<FInterchangeImportStatus>(FilesAndDestinations.Num());
 
+	const bool bForceContentBrowserSyncIfOnlyOneMainAsset = (FilesAndDestinations.Num() == 1 && !bAutomatedImport);
+
 	// Now iterate over the input files and use the same factory object for each file with the same extension
 	for(int32 FileIdx = 0; FileIdx < FilesAndDestinations.Num() && !bImportWasCancelled; ++FileIdx)
 	{
@@ -3204,10 +3209,10 @@ TArray<UObject*> UAssetToolsImpl::ImportAssetsInternal(const TArray<FString>& Fi
 							UEditorLoadingAndSavingUtils::SavePackages(PackagesToSave, bDirtyOnly);
 						}
 					};
-
+				
 				TFunction<void(UE::Interchange::FImportResult&)> AppendAndBroadcastImportResultIfNeeded =
 					// Note: ImportStatus captured by value so that the lambda keeps the shared ptr alive
-					[ImportStatus, AppendImportResult, bSyncToBrowser](UE::Interchange::FImportResult& Result)
+					[ImportStatus, AppendImportResult, bSyncToBrowser, bForceContentBrowserSyncIfOnlyOneMainAsset](UE::Interchange::FImportResult& Result)
 					{
 						AppendImportResult(Result);
 
@@ -3216,10 +3221,21 @@ TArray<UObject*> UAssetToolsImpl::ImportAssetsInternal(const TArray<FString>& Fi
 							UInterchangeManager& InterchangeManager = UInterchangeManager::GetInterchangeManager();
 							InterchangeManager.OnBatchImportComplete.Broadcast(ImportStatus->InterchangeResultsContainer);
 
-							if (bSyncToBrowser)
+							TArray<UObject*> MainAssets;
+							for (const TWeakObjectPtr<UObject>& WeakObject : ImportStatus->ImportedObjects)
+							{
+								if (WeakObject.IsValid()
+									&& (WeakObject->IsA<UStaticMesh>()
+										|| WeakObject->IsA<USkeletalMesh>()
+										|| WeakObject->IsA<UAnimSequence>()))
+								{
+									MainAssets.Add(WeakObject.Get());
+								}
+							}
+
+							//Force browser to sync to the import asset if there is only one asset imported
+							if (bSyncToBrowser || (bForceContentBrowserSyncIfOnlyOneMainAsset && MainAssets.Num() == 1))
 							{
-								// Only sync the content browser when the full import is done. Otherwise it can be annoying for the user.
-								// UX suggestion : Maybe we could move this into the post import window as button ("Select Imported Asset(s)") so it would be even less disruptive for the users.
 								TArray<UObject*> ImportedObjects;
 								ImportedObjects.Reserve(ImportStatus->ImportedObjects.Num());
 								for (const TWeakObjectPtr<UObject>& WeakObject : ImportStatus->ImportedObjects)
@@ -3227,7 +3243,7 @@ TArray<UObject*> UAssetToolsImpl::ImportAssetsInternal(const TArray<FString>& Fi
 									ImportedObjects.Add(WeakObject.Get());
 								}
 
-								UAssetToolsImpl::Get().SyncBrowserToAssets(ImportedObjects);
+								UAssetToolsImpl::Get().SyncBrowserToAssets(bSyncToBrowser ? ImportedObjects : MainAssets);
 							}
 						}
 					};
@@ -3623,7 +3639,7 @@ TArray<UObject*> UAssetToolsImpl::ImportAssetsInternal(const TArray<FString>& Fi
 	}
 
 	// Sync content browser to the newly created assets
-	if(ReturnObjects.Num() && (bSyncToBrowser != false))
+	if(ReturnObjects.Num() && (bSyncToBrowser != false || (ReturnObjects.Num() == 1 && bForceContentBrowserSyncIfOnlyOneMainAsset)))
 	{
 		UAssetToolsImpl::Get().SyncBrowserToAssets(ReturnObjects);
 	}
