AssetDefinition - When proxying to the old category system, if the advanced top level category doesn't exist, make sure to create it.  Fixes FX/Niagara stuff not appearing in the top level menu.

diff --git a/Engine/Source/Developer/AssetTools/Private/AssetTools.cpp b/Engine/Source/Developer/AssetTools/Private/AssetTools.cpp
--- a/Engine/Source/Developer/AssetTools/Private/AssetTools.cpp
+++ b/Engine/Source/Developer/AssetTools/Private/AssetTools.cpp
@@ -188,7 +188,22 @@ public:
 			}
 			else
 			{
-				CategoryBits |= AssetTools.FindAdvancedAssetCategory(Category.GetCategory());	
+				if (Category.GetCategory() != EAssetCategoryPaths::Misc.GetCategory())
+				{
+					EAssetTypeCategories::Type AdvancedCategoryBit = AssetTools.FindAdvancedAssetCategory(Category.GetCategory());
+					if (AdvancedCategoryBit == EAssetTypeCategories::Misc)
+					{
+						CategoryBits |= AssetTools.RegisterAdvancedAssetCategory(Category.GetCategory(), Category.GetCategoryText());
+					}
+					else
+					{
+						CategoryBits |= AdvancedCategoryBit;
+					}
+				}
+				else
+				{
+					CategoryBits |= EAssetTypeCategories::Misc;
+				} 
 			}
 		}
 
@@ -1324,7 +1339,6 @@ UAssetToolsImpl::UAssetToolsImpl(const FObjectInitializer& ObjectInitializer)
 	EAssetTypeCategories::Type InputCategoryBit = RegisterAdvancedAssetCategory(FName(TEXT("Input")), LOCTEXT("InputAssetsCategory", "Input"));
 
 	// Register the built-in asset type actions
-	RegisterAssetTypeActions(MakeShareable(new FAssetTypeActions_AnimationAsset));
 	RegisterAssetTypeActions(MakeShareable(new FAssetTypeActions_AnimBlueprintInterface));
 	RegisterAssetTypeActions(MakeShareable(new FAssetTypeActions_AnimBlueprint));
 	RegisterAssetTypeActions(MakeShareable(new FAssetTypeActions_AnimBlueprintGeneratedClass));
