Fixed an issue where the metadata of the package wouldn't be migrated for some assets.

diff --git a/Engine/Source/Developer/AssetTools/Private/AssetTools.cpp b/Engine/Source/Developer/AssetTools/Private/AssetTools.cpp
--- a/Engine/Source/Developer/AssetTools/Private/AssetTools.cpp
+++ b/Engine/Source/Developer/AssetTools/Private/AssetTools.cpp
@@ -1016,6 +1016,7 @@ namespace UE::AssetTools::Private
 		static void SaveInstancedPackagesIntoDestination(FPackageMigrationContext& PackageMigrationContext, FPackageMigrationImplContext& MigrationImplContext)
 		{
 			FSavePackageArgs SaveArgs;
+			SaveArgs.TopLevelFlags = RF_Standalone;
 
 			// Rehydrate the payloads so that we move the virtual payloads also. Add the auto save flag to notify the editor extensions that they shouldn't prompt the user or generate/update some asset on save.
 			SaveArgs.SaveFlags |= SAVE_RehydratePayloads | SAVE_FromAutosave;
