Advanced Copy - Optional callback when copy is complete, and expose to blueprints

* Adds a new delegate for when AdvancedCopy is complete (since Advanced Copy is async and can't return the results directly)
* Adds a new blueprint callable overload for AdvancedCopy in the AssetTools interface that optionally takes this callback

diff --git a/Engine/Source/Developer/AssetTools/Private/AssetTools.cpp b/Engine/Source/Developer/AssetTools/Private/AssetTools.cpp
--- a/Engine/Source/Developer/AssetTools/Private/AssetTools.cpp
+++ b/Engine/Source/Developer/AssetTools/Private/AssetTools.cpp
@@ -1962,10 +1962,10 @@ FString UAssetToolsImpl::GenerateAdvancedCopyDestinationPackageName(const FStrin
 	return DestinationPackageName;
 }
 
-bool UAssetToolsImpl::FlattenAdvancedCopyDestinations(const TArray<TMap<FString, FString>> PackagesAndDestinations, TMap<FString, FString>& FlattenedPackagesAndDestinations) const
+bool UAssetToolsImpl::FlattenAdvancedCopyDestinations(const TArray<TMap<FString, FString>>& PackagesAndDestinations, TMap<FString, FString>& FlattenedPackagesAndDestinations) const
 {
 	FString CopyErrors;
-	for (TMap<FString, FString> PackageAndDestinationMap : PackagesAndDestinations)
+	for (const TMap<FString, FString>& PackageAndDestinationMap : PackagesAndDestinations)
 	{
 		for (auto It = PackageAndDestinationMap.CreateConstIterator(); It; ++It)
 		{
@@ -2293,14 +2293,34 @@ bool UAssetToolsImpl::AdvancedCopyPackages(
 	return false;
 }
 
-bool UAssetToolsImpl::AdvancedCopyPackages(const FAdvancedCopyParams& CopyParams, const TArray<TMap<FString, FString>> PackagesAndDestinations) const
+bool UAssetToolsImpl::AdvancedCopyPackages(const FAdvancedCopyParams& CopyParams, const TArray<TMap<FString, FString>>& PackagesAndDestinations) const
 {
+	bool bResult = false;
+	FDuplicatedObjects DuplicatedObjects;
+	
 	TMap<FString, FString> FlattenedDestinationMap;
 	if (FlattenAdvancedCopyDestinations(PackagesAndDestinations, FlattenedDestinationMap))
 	{
-		return AdvancedCopyPackages(FlattenedDestinationMap, CopyParams.bShouldForceSave, CopyParams.bCopyOverAllDestinationOverlaps, /*OutDuplicatedObjects=*/nullptr, EMessageSeverity::Info);
+		bResult = AdvancedCopyPackages(FlattenedDestinationMap, CopyParams.bShouldForceSave, CopyParams.bCopyOverAllDestinationOverlaps, &DuplicatedObjects, EMessageSeverity::Info);
 	}
-	return false;
+
+	if (CopyParams.OnCopyComplete.IsBound())
+	{
+		TArray<FAssetRenameData> AllCopiedAssets;
+		for (const TMap<TSoftObjectPtr<UObject>, TSoftObjectPtr<UObject>>& PerPackageCopies : DuplicatedObjects)
+		{
+			for (const TPair<TSoftObjectPtr<UObject>, TSoftObjectPtr<UObject>>& AssetSourceAndCopy : PerPackageCopies)
+			{
+				const TSoftObjectPtr<UObject>& CopiedAsset = AssetSourceAndCopy.Value;
+				FAssetRenameData& AssetRenameData = AllCopiedAssets.Emplace_GetRef(TWeakObjectPtr(CopiedAsset.Get()), CopiedAsset.GetLongPackageName(), CopiedAsset.GetAssetName());
+				AssetRenameData.OldObjectPath = AssetSourceAndCopy.Key.ToSoftObjectPath();
+				AssetRenameData.NewObjectPath = CopiedAsset.ToSoftObjectPath();
+			}
+		}
+		CopyParams.OnCopyComplete.Execute(bResult, AllCopiedAssets);
+	}
+	
+	return bResult;
 }
 
 bool UAssetToolsImpl::RenameAssets(const TArray<FAssetRenameData>& AssetsAndNames)
@@ -5239,6 +5259,11 @@ void UAssetToolsImpl::ConvertVirtualTextures(const TArray<UTexture2D *>& Texture
 }
 
 void UAssetToolsImpl::BeginAdvancedCopyPackages(const TArray<FName>& InputNamesToCopy, const FString& TargetPath) const
+{
+	BeginAdvancedCopyPackages(InputNamesToCopy, TargetPath, FAdvancedCopyCompletedEvent());
+}
+
+void UAssetToolsImpl::BeginAdvancedCopyPackages(const TArray<FName>& InputNamesToCopy, const FString& TargetPath, const FAdvancedCopyCompletedEvent& OnCopyComplete) const
 {
 	// Packages must be saved for the migration to work
 	const bool bPromptUserToSave = true;
@@ -5251,13 +5276,13 @@ void UAssetToolsImpl::BeginAdvancedCopyPackages(const TArray<FName>& InputNamesT
 		{
 			// Open a dialog asking the user to wait while assets are being discovered
 			SDiscoveringAssetsDialog::OpenDiscoveringAssetsDialog(
-				SDiscoveringAssetsDialog::FOnAssetsDiscovered::CreateUObject(this, &UAssetToolsImpl::PerformAdvancedCopyPackages, InputNamesToCopy, TargetPath)
+				SDiscoveringAssetsDialog::FOnAssetsDiscovered::CreateUObject(this, &UAssetToolsImpl::PerformAdvancedCopyPackages, InputNamesToCopy, TargetPath, OnCopyComplete)
 			);
 		}
 		else
 		{
 			// Assets are already discovered, perform the migration now
-			PerformAdvancedCopyPackages(InputNamesToCopy, TargetPath);
+			PerformAdvancedCopyPackages(InputNamesToCopy, TargetPath, OnCopyComplete);
 		}
 	}
 }
@@ -5297,13 +5322,14 @@ TArray<FName> UAssetToolsImpl::ExpandAssetsAndFoldersToJustAssets(TArray<FName>
 	return ExpandedAssetArray;
 }
 
-void UAssetToolsImpl::PerformAdvancedCopyPackages(TArray<FName> SelectedAssetAndFolderNames, FString TargetPath) const
+void UAssetToolsImpl::PerformAdvancedCopyPackages(TArray<FName> SelectedAssetAndFolderNames, FString TargetPath, FAdvancedCopyCompletedEvent OnCopyComplete) const
 {	
 	TargetPath.RemoveFromEnd(TEXT("/"));
 
 	//TArray<FName> ExpandedAssets = ExpandAssetsAndFoldersToJustAssets(SelectedAssetAndFolderNames);
 	FAdvancedCopyParams CopyParams = FAdvancedCopyParams(SelectedAssetAndFolderNames, TargetPath);
 	CopyParams.bShouldCheckForDependencies = true;
+	CopyParams.OnCopyComplete = OnCopyComplete;
 
 	// Suppress UI if we're running in unattended mode
 	if (FApp::IsUnattended())
@@ -5419,7 +5445,7 @@ void UAssetToolsImpl::InitAdvancedCopyFromCopyParams(FAdvancedCopyParams CopyPar
 	}
 }
 
-void UAssetToolsImpl::AdvancedCopyPackages_ReportConfirmed(FAdvancedCopyParams CopyParams, TArray<TMap<FString, FString>> DestinationMap) const
+void UAssetToolsImpl::AdvancedCopyPackages_ReportConfirmed(const FAdvancedCopyParams& CopyParams, const TArray<TMap<FString, FString>>& DestinationMap) const
 {
 	TArray<UAdvancedCopyCustomization*> CustomizationsToUse = CopyParams.GetCustomizationsToUse();
 	for (int32 CustomizationIndex = 0; CustomizationIndex < CustomizationsToUse.Num(); CustomizationIndex++)
diff --git a/Engine/Source/Developer/AssetTools/Private/AssetTools.h b/Engine/Source/Developer/AssetTools/Private/AssetTools.h
--- a/Engine/Source/Developer/AssetTools/Private/AssetTools.h
+++ b/Engine/Source/Developer/AssetTools/Private/AssetTools.h
@@ -113,14 +113,15 @@ public:
 	virtual void MigratePackages(const TArray<FName>& PackageNamesToMigrate, const FString& TargetPath, const struct FMigrationOptions& Options = FMigrationOptions()) const override;
 	virtual UE::AssetTools::FOnPackageMigration& GetOnPackageMigration() override;
 	virtual void BeginAdvancedCopyPackages(const TArray<FName>& InputNamesToCopy, const FString& TargetPath) const override;
+	virtual void BeginAdvancedCopyPackages(const TArray<FName>& InputNamesToCopy, const FString& TargetPath, const FAdvancedCopyCompletedEvent& OnCopyComplete) const override;
 	virtual void FixupReferencers(const TArray<UObjectRedirector*>& Objects, bool bCheckoutDialogPrompt = true, ERedirectFixupMode FixupMode = ERedirectFixupMode::DeleteFixedUpRedirectors) const override;
 	virtual bool IsFixupReferencersInProgress() const override;
 	virtual FAssetPostRenameEvent& OnAssetPostRename() override { return AssetRenameManager->OnAssetPostRenameEvent(); }
 	virtual void ExpandDirectories(const TArray<FString>& Files, const FString& DestinationPath, TArray<TPair<FString, FString>>& FilesAndDestinations) const override;
-	virtual bool AdvancedCopyPackages(const FAdvancedCopyParams& CopyParams, const TArray<TMap<FString, FString>> PackagesAndDestinations) const override;
+	virtual bool AdvancedCopyPackages(const FAdvancedCopyParams& CopyParams, const TArray<TMap<FString, FString>>& PackagesAndDestinations) const override;
 	virtual bool AdvancedCopyPackages(const TMap<FString, FString>& SourceAndDestPackages, const bool bForceAutosave, const bool bCopyOverAllDestinationOverlaps, FDuplicatedObjects* OutDuplicatedObjects, EMessageSeverity::Type NotificationSeverityFilter) const override;
 	virtual void GenerateAdvancedCopyDestinations(FAdvancedCopyParams& InParams, const TArray<FName>& InPackageNamesToCopy, const class UAdvancedCopyCustomization* CopyCustomization, TMap<FString, FString>& OutPackagesAndDestinations) const override;
-	virtual bool FlattenAdvancedCopyDestinations(const TArray<TMap<FString, FString>> PackagesAndDestinations, TMap<FString, FString>& FlattenedPackagesAndDestinations) const override;
+	virtual bool FlattenAdvancedCopyDestinations(const TArray<TMap<FString, FString>>& PackagesAndDestinations, TMap<FString, FString>& FlattenedPackagesAndDestinations) const override;
 	virtual bool ValidateFlattenedAdvancedCopyDestinations(const TMap<FString, FString>& FlattenedPackagesAndDestinations) const override;
 	virtual void GetAllAdvancedCopySources(FName SelectedPackage, FAdvancedCopyParams& CopyParams, TArray<FName>& OutPackageNamesToCopy, TMap<FName, FName>& DependencyMap, const class UAdvancedCopyCustomization* CopyCustomization) const override;
 	virtual void InitAdvancedCopyFromCopyParams(FAdvancedCopyParams CopyParams) const override;
@@ -175,13 +176,13 @@ private:
 	TArray<FName> ExpandAssetsAndFoldersToJustAssets(TArray<FName> SelectedAssetAndFolderNames) const;
 
 	/** Begins the package advanced copy, after assets have been discovered */
-	void PerformAdvancedCopyPackages(TArray<FName> SelectedPackageNames, FString TargetPath) const;
+	void PerformAdvancedCopyPackages(TArray<FName> SelectedPackageNames, FString TargetPath, FAdvancedCopyCompletedEvent OnCopyComplete) const;
 
 	/** Copies files after the final list was confirmed */
 	void MigratePackages_ReportConfirmed(TSharedPtr<TArray<ReportPackageData>> PackageDataToMigrate, const FString DestinationPath, TSet<FName> ExcludedDependencies, const FMigrationOptions Options) const;
 
 	/** Copies files after the final list was confirmed */
-	void AdvancedCopyPackages_ReportConfirmed(FAdvancedCopyParams CopyParam, TArray<TMap<FString, FString>> DestinationMap) const;
+	void AdvancedCopyPackages_ReportConfirmed(const FAdvancedCopyParams& CopyParam, const TArray<TMap<FString, FString>>& DestinationMap) const;
 
 	/** Gets the dependencies of the specified package recursively */
 	void RecursiveGetDependencies(const FName& PackageName, TSet<FName>& AllDependencies, TSet<FString>& ExternalObjectsPaths, TSet<FName>& ExcludedDependencies, const TFunction<bool(FName)>& ShouldExcludeFromDependenciesSearch) const;
diff --git a/Engine/Source/Developer/AssetTools/Private/SAdvancedCopyReportDialog.h b/Engine/Source/Developer/AssetTools/Private/SAdvancedCopyReportDialog.h
--- a/Engine/Source/Developer/AssetTools/Private/SAdvancedCopyReportDialog.h
+++ b/Engine/Source/Developer/AssetTools/Private/SAdvancedCopyReportDialog.h
@@ -72,7 +72,7 @@ class SAdvancedCopyReportDialog : public SCompoundWidget
 {
 public:
 	typedef TArray<TMap<FString, FString>> FCopyDestinationMap;
-	DECLARE_DELEGATE_TwoParams(FOnReportConfirmed, FAdvancedCopyParams, FCopyDestinationMap)
+	DECLARE_DELEGATE_TwoParams(FOnReportConfirmed, const FAdvancedCopyParams&, const FCopyDestinationMap&)
 
 	SLATE_BEGIN_ARGS(SAdvancedCopyReportDialog){}
 
diff --git a/Engine/Source/Developer/AssetTools/Public/IAssetTools.h b/Engine/Source/Developer/AssetTools/Public/IAssetTools.h
--- a/Engine/Source/Developer/AssetTools/Public/IAssetTools.h
+++ b/Engine/Source/Developer/AssetTools/Public/IAssetTools.h
@@ -118,6 +118,7 @@ struct FAssetRenameData
 	}
 };
 
+DECLARE_DYNAMIC_DELEGATE_TwoParams(FAdvancedCopyCompletedEvent, bool, bSuccess, const TArray<FAssetRenameData>&, AllCopiedAssets);
 
 DECLARE_MULTICAST_DELEGATE_OneParam(FAssetPostRenameEvent, const TArray<FAssetRenameData>&);
 DECLARE_DELEGATE_RetVal_TwoParams(bool, FIsNameAllowed, const FString& /*Name*/, FText* /*OutErrorMessage*/);
@@ -144,7 +145,8 @@ struct FAdvancedCopyParams
 	bool bCopyOverAllDestinationOverlaps = false;
 	bool bShouldSuppressUI = false;
 	bool bShouldCheckForDependencies = false;
-
+	FAdvancedCopyCompletedEvent OnCopyComplete;
+	
 	UE_DEPRECATED(5.0, "This function has been deprecated, use GetSelectedPackageOrFolderNames")
 	const TArray<FName>& GetSelectedPackageNames() const
 	{
@@ -506,6 +508,9 @@ public:
 	/* Copy packages and dependencies to another folder */
 	virtual void BeginAdvancedCopyPackages(const TArray<FName>& InputNamesToCopy, const FString& TargetPath) const = 0;
 
+	UFUNCTION(BlueprintCallable, Category = "Editor Scripting | Asset Tools", BlueprintPure = false, meta = (AutoCreateRefTerm = "OnCopyComplete"))
+	virtual void BeginAdvancedCopyPackages(const TArray<FName>& InputNamesToCopy, const FString& TargetPath, const FAdvancedCopyCompletedEvent& OnCopyComplete) const = 0;
+
 	/**
 	 * Fix up references to the specified redirectors.
 	 * @param bCheckoutDialogPrompt indicates whether to prompt the user with files checkout dialog or silently attempt to checkout all necessary files.
@@ -519,7 +524,7 @@ public:
 	virtual void ExpandDirectories(const TArray<FString>& Files, const FString& DestinationPath, TArray<TPair<FString, FString>>& FilesAndDestinations) const = 0;
 
 	/** Copies files after the final set of maps of sources and destinations was confirmed */
-	virtual bool AdvancedCopyPackages(const FAdvancedCopyParams& CopyParams, const TArray<TMap<FString, FString>> PackagesAndDestinations) const = 0;
+	virtual bool AdvancedCopyPackages(const FAdvancedCopyParams& CopyParams, const TArray<TMap<FString, FString>>& PackagesAndDestinations) const = 0;
 
 	/** Copies files after the flattened map of sources and destinations was confirmed */
 	virtual bool AdvancedCopyPackages(const TMap<FString, FString>& SourceAndDestPackages, const bool bForceAutosave = false, const bool bCopyOverAllDestinationOverlaps = true, FDuplicatedObjects* OutDuplicatedObjects = nullptr, EMessageSeverity::Type NotificationSeverityFilter = EMessageSeverity::Info) const = 0;
@@ -528,7 +533,7 @@ public:
 	virtual void GenerateAdvancedCopyDestinations(FAdvancedCopyParams& InParams, const TArray<FName>& InPackageNamesToCopy, const UAdvancedCopyCustomization* CopyCustomization, TMap<FString, FString>& OutPackagesAndDestinations) const = 0;
 
 	/* Flattens the maps for each selected package into one complete map to pass to the final copy function while checking for collisions */
-	virtual bool FlattenAdvancedCopyDestinations(const TArray<TMap<FString, FString>> PackagesAndDestinations, TMap<FString, FString>& FlattenedPackagesAndDestinations) const = 0;
+	virtual bool FlattenAdvancedCopyDestinations(const TArray<TMap<FString, FString>>& PackagesAndDestinations, TMap<FString, FString>& FlattenedPackagesAndDestinations) const = 0;
 
 	/* Validate the destinations for advanced copy once the map has been flattened */
 	virtual bool ValidateFlattenedAdvancedCopyDestinations(const TMap<FString, FString>& FlattenedPackagesAndDestinations) const = 0;
