Fixed: SkeinUE does not sync latest if only non-package files have changed

diff --git a/Engine/Source/Developer/AssetTools/Private/AssetViewUtils.cpp b/Engine/Source/Developer/AssetTools/Private/AssetViewUtils.cpp
--- a/Engine/Source/Developer/AssetTools/Private/AssetViewUtils.cpp
+++ b/Engine/Source/Developer/AssetTools/Private/AssetViewUtils.cpp
@@ -1706,7 +1706,7 @@ void ShowSyncDependenciesDialog(const TArray<FString>& InDependencies, TArray<FS
 
 void AssetViewUtils::SyncPackagesFromSourceControl(const TArray<FString>& PackageNames, bool bIsSyncLatestOperation)
 {
-	if (PackageNames.Num() > 0)
+	if (PackageNames.Num() > 0 || bIsSyncLatestOperation)
 	{
 		TArray<FString> PackageNamesToSync = PackageNames;
 
