Fixed spelling typo in log message.

diff --git a/Engine/Source/Developer/AssetTools/Private/AssetTools.cpp b/Engine/Source/Developer/AssetTools/Private/AssetTools.cpp
--- a/Engine/Source/Developer/AssetTools/Private/AssetTools.cpp
+++ b/Engine/Source/Developer/AssetTools/Private/AssetTools.cpp
@@ -1040,7 +1040,7 @@ namespace UE::AssetTools::Private
 
 						if (SaveResult.IsSuccessful())
 						{
-							PackageMigrationContext.AddSucessfullMigrationMessage(FText::Format(LOCTEXT("MigratePackages_SaveSuccess", "Package ({0}) was migrated successfully migrated as ({1}) with the fellowing filename ({2})")\
+							PackageMigrationContext.AddSucessfullMigrationMessage(FText::Format(LOCTEXT("MigratePackages_SaveSuccess", "Package ({0}) was migrated successfully as ({1}) with the following filename ({2})")\
 								, FText::FromString(MigrationPackageData.GetOriginalPackageName())
 								, FText::FromString(MigrationPackageData.GetInstancedPackageName())
 								, FText::FromString(MigrationPackageData.GetDestinationFilename())
