SourceControl: Added notification if sync fails when getting latest, similar to the warning shown by FSourceControlWindows::PromptForCheckin for checkin fails.

In addition, I refactored AssetViewUtils::SyncPathsFromSourceControl somewhat.
- Added alternative (preferred) way to get affected files/packages using the FSyncPreview operation as opposed to querying the AssetRegistry.
- Added support for hot reloading the world if one of its external packages (actor/object) got synced so those changes are reflected in the viewport.
- Added support for passing in either content paths or absolute paths to sync.

- Deprecated the 'bIsSyncLatestOperation' parameter from SyncPackagesFromSourceControl.

Note that SyncPackagesFromSourceControl now has a single internal caller (FAssetSourceControlContextMenuState::ExecuteSCCSync).
Note that SyncPathsFromSourceControl now has a single internal caller (FSourceControlWindows::SyncLatest).

diff --git a/Engine/Source/Developer/AssetTools/Private/AssetViewUtils.cpp b/Engine/Source/Developer/AssetTools/Private/AssetViewUtils.cpp
--- a/Engine/Source/Developer/AssetTools/Private/AssetViewUtils.cpp
+++ b/Engine/Source/Developer/AssetTools/Private/AssetViewUtils.cpp
@@ -1705,21 +1705,30 @@ void ShowSyncDependenciesDialog(const TArray<FString>& InDependencies, TArray<FS
 
 void AssetViewUtils::SyncPackagesFromSourceControl(const TArray<FString>& PackageNames, bool bIsSyncLatestOperation)
 {
-	if (PackageNames.Num() > 0 || bIsSyncLatestOperation)
+	if (bIsSyncLatestOperation)
+	{
+		SyncLatestFromSourceControl();
+	}
+	else
+	{
+		SyncPackagesFromSourceControl(PackageNames);
+	}
+}
+
+bool AssetViewUtils::SyncPackagesFromSourceControl(const TArray<FString>& PackageNames)
+{
+	if (PackageNames.Num() > 0)
 	{
 		TArray<FString> PackageNamesToSync = PackageNames;
 
-		if (!bIsSyncLatestOperation)
-		{
-			// Warn about any packages that are being synced without also getting the newest version of their dependencies...
-			TArray<FString> OutOfDateDependencies;
-			GetOutOfDatePackageDependencies(PackageNamesToSync, OutOfDateDependencies);
+		// Warn about any packages that are being synced without also getting the newest version of their dependencies...
+		TArray<FString> OutOfDateDependencies;
+		GetOutOfDatePackageDependencies(PackageNamesToSync, OutOfDateDependencies);
 
-			TArray<FString> ExtraPackagesToSync;
-			ShowSyncDependenciesDialog(OutOfDateDependencies, ExtraPackagesToSync);
+		TArray<FString> ExtraPackagesToSync;
+		ShowSyncDependenciesDialog(OutOfDateDependencies, ExtraPackagesToSync);
 
-			PackageNamesToSync.Append(ExtraPackagesToSync);
-		}
+		PackageNamesToSync.Append(ExtraPackagesToSync);
 
 		ISourceControlProvider& SCCProvider = ISourceControlModule::Get().GetProvider();
 		const TArray<FString> PackageFilenames = SourceControlHelpers::PackageFilenames(PackageNamesToSync);
@@ -1746,19 +1755,7 @@ void AssetViewUtils::SyncPackagesFromSourceControl(const TArray<FString>& Packag
 
 		// Sync everything...
 		TSharedRef<FSync> Operation = ISourceControlOperation::Create<FSync>();
-		if (bIsSyncLatestOperation)
-		{
-			Operation->SetHeadRevisionFlag(true);
-			if (SCCProvider.Execute(Operation) != ECommandResult::Succeeded)
-			{
-				UE_LOG(LogAssetViewTools, Warning, TEXT("Error while syncing latest revision"));
-				return;
-			}
-		}
-		else
-		{
-			SCCProvider.Execute(Operation, PackageFilenames);
-		}
+		ECommandResult::Type SyncResult = SCCProvider.Execute(Operation, PackageFilenames);
 
 		// Syncing may have deleted some packages, so we need to unload those rather than re-load them...
 		// Note: we will store the package using weak pointers here otherwise we might have garbage collection issues after the ReloadPackages call
@@ -1776,11 +1773,6 @@ void AssetViewUtils::SyncPackagesFromSourceControl(const TArray<FString>& Packag
 		});
 
 		// Hot-reload the new packages...
-		if (bIsSyncLatestOperation)
-		{
-			UWorld* EditorWorld = GEditor->GetEditorWorldContext().World();
-			LoadedPackages.Add(EditorWorld->GetPackage());
-		}
 		UPackageTools::ReloadPackages(LoadedPackages);
 
 		// Unload any deleted packages...
@@ -1797,53 +1789,32 @@ void AssetViewUtils::SyncPackagesFromSourceControl(const TArray<FString>& Packag
 
 		// Re-cache the SCC state...
 		SCCProvider.Execute(ISourceControlOperation::Create<FUpdateStatus>(), PackageFilenames);
-	}
-}
-
-void AssetViewUtils::SyncLatestFromSourceControl()
-{
-	TArray<FString> PackageNames;
-	TArray<FString> PackageFilenames;
-
-	ISourceControlProvider& SCCProvider = ISourceControlModule::Get().GetProvider();
-	TSharedRef<FSyncPreview> PreviewOperation = ISourceControlOperation::Create<FSyncPreview>();
-	PreviewOperation->SetHeadRevisionFlag(true);
-	ECommandResult::Type OperationResult =  SCCProvider.Execute(PreviewOperation);
-
-	if (OperationResult == ECommandResult::Succeeded)
-	{
-		PackageFilenames = PreviewOperation->GetAffectedFiles();
-
-		for (const FString& PackageFilename : PackageFilenames)
-		{
-			FString PackageName;
-			if (FPackageName::TryConvertFilenameToLongPackageName(PackageFilename, PackageName))
-			{
-				PackageNames.Add(PackageName);
-			}
-		}
 
-		SyncPackagesFromSourceControl(PackageNames, /*IsSyncLatestOperation=*/true);
+		// Return result...
+		return (SyncResult == ECommandResult::Succeeded);
 	}
-	else
-	{
-		// Likely the source control provider does not support sync previews. Proceed to syncing all paths under source control
 
-		bool bContentOnly = !ISourceControlModule::Get().UsesCustomProjectDir();
-		TArray<FString> SourceControlLocations = SourceControlHelpers::GetSourceControlLocations(bContentOnly);
-
-		SyncPathsFromSourceControl(SourceControlLocations);
-	}
+	return true;
 }
 
-void AssetViewUtils::SyncPathsFromSourceControl(const TArray<FString>& ContentPaths)
+static bool SyncPathsFromSourceControl(const TArray<FString>& Paths, bool bCheckDependencies)
 {
 	TArray<FString> PathsOnDisk;
-	PathsOnDisk.Reserve(ContentPaths.Num());
-	for (const FString& ContentPath : ContentPaths)
+	PathsOnDisk.Reserve(Paths.Num());
+	for (const FString& Path : Paths)
 	{
 		FString PathOnDisk;
-		if (FPackageName::TryConvertLongPackageNameToFilename(ContentPath / TEXT(""), PathOnDisk) && FPaths::DirectoryExists(PathOnDisk))
+		if (FPackageName::TryConvertLongPackageNameToFilename(Path / TEXT(""), PathOnDisk))
+		{
+			// Given path was a content path.
+		}
+		else
+		{
+			// Treat path as a (relative or absolute) disk path.
+			PathOnDisk = FPaths::ConvertRelativePathToFull(Path);
+		}
+
+		if (FPaths::DirectoryExists(PathOnDisk))
 		{
 			PathsOnDisk.Emplace(MoveTemp(PathOnDisk));
 		}
@@ -1851,22 +1822,51 @@ void AssetViewUtils::SyncPathsFromSourceControl(const TArray<FString>& ContentPa
 
 	if (PathsOnDisk.Num() > 0)
 	{
-		// Get all the assets under the path(s) on disk...
+		ISourceControlProvider& SCCProvider = ISourceControlModule::Get().GetProvider();
+
+		// Get all the assets about to be synced in those path(s) on disk...
+		TArray<FString> AffectedFiles;
 		TArray<FString> PackageNames;
+
+		// Use FSyncPreview if possible.
+		TSharedRef<FSyncPreview> PreviewOperation = ISourceControlOperation::Create<FSyncPreview>();
+		if (SCCProvider.CanExecuteOperation(PreviewOperation) &&
+			SCCProvider.Execute(PreviewOperation, PathsOnDisk) == ECommandResult::Succeeded)
 		{
+			// The source control provider supports sync previews which gets us the full list of
+			// affected files, both packages and non-packages.
+			AffectedFiles = PreviewOperation->GetAffectedFiles();
+
+			// Determine which packages would be affected.
+			for (const FString& FileName : AffectedFiles)
+			{
+				FString PackageName;
+				if (FPackageName::TryConvertFilenameToLongPackageName(FileName, PackageName))
+				{
+					PackageNames.Add(PackageName);
+				}
+			}
+		}
+		else
+		{
+			// The source control provider does not support sync previews.
+			// Fallback on using the AssetRegistry.
 			FAssetRegistryModule& AssetRegistryModule = FModuleManager::Get().LoadModuleChecked<FAssetRegistryModule>(TEXT("AssetRegistry"));
 
 			FARFilter Filter;
 			Filter.bRecursivePaths = true;
 			for (const FString& PathOnDisk : PathsOnDisk)
 			{
-				FString PackagePath = FPackageName::FilenameToLongPackageName(PathOnDisk);
-				if (PackagePath.Len() > 1 && PackagePath[PackagePath.Len() - 1] == TEXT('/'))
+				FString PackagePath;
+				if (FPackageName::TryConvertFilenameToLongPackageName(PathOnDisk, PackagePath))
 				{
-					// The filter path can't end with a trailing slash
-					PackagePath.LeftChopInline(1, false);
+					if (PackagePath.Len() > 1 && PackagePath[PackagePath.Len() - 1] == TEXT('/'))
+					{
+						// The filter path can't end with a trailing slash
+						PackagePath.LeftChopInline(1, false);
+					}
+					Filter.PackagePaths.Emplace(*PackagePath);
 				}
-				Filter.PackagePaths.Emplace(*PackagePath);
 			}
 
 			TArray<FAssetData> AssetList;
@@ -1884,11 +1884,10 @@ void AssetViewUtils::SyncPathsFromSourceControl(const TArray<FString>& ContentPa
 			}
 		}
 
-		ISourceControlProvider& SCCProvider = ISourceControlModule::Get().GetProvider();
-
 		// Warn about any packages that are being synced without also getting the newest version of their dependencies...
 		TArray<FString> PackageNamesToSync = PackageNames;
 		TArray<FString> ExtraPackagesToSync;
+		if (bCheckDependencies)
 		{
 			TArray<FString> OutOfDateDependencies;
 			GetOutOfDatePackageDependencies(PackageNamesToSync, OutOfDateDependencies);
@@ -1925,7 +1924,7 @@ void AssetViewUtils::SyncPathsFromSourceControl(const TArray<FString>& ContentPa
 		PathsToSync.Append(SourceControlHelpers::PackageFilenames(ExtraPackagesToSync));
 
 		// Sync everything...
-		SCCProvider.Execute(ISourceControlOperation::Create<FSync>(), PathsToSync);
+		ECommandResult::Type SyncResult = SCCProvider.Execute(ISourceControlOperation::Create<FSync>(), PathsToSync);
 
 		// Syncing may have deleted some packages, so we need to unload those rather than re-load them...
 		// Note: we will store the package using weak pointers here otherwise we might have garbage collection issues after the ReloadPackages call
@@ -1942,11 +1941,23 @@ void AssetViewUtils::SyncPathsFromSourceControl(const TArray<FString>& ContentPa
 			return false; // keep package
 		});
 
-		//UE_LOG(LogContentBrowser, Log, TEXT("Syncing %d path(s):"), ContentPaths.Num());
-		//for (const UPackage* Package : LoadedPackages)
-		//{
-		//	UE_LOG(LogContentBrowser, Log, TEXT("\t - %s"), *Package->GetName());
-		//}
+		// Check if the world should be reloaded as well because one of its external packages got synced...
+		if (UWorld* EditorWorld = GEditor->GetEditorWorldContext().World())
+		{
+			UPackage* EditorWorldPackage = EditorWorld->GetPackage();
+			if (!LoadedPackages.Contains(EditorWorldPackage))
+			{
+				for (const FString& PackageName : PackageNamesToSync)
+				{
+					if (PackageName.Contains(FPackagePath::GetExternalActorsFolderName()) ||
+						PackageName.Contains(FPackagePath::GetExternalObjectsFolderName()))
+					{
+						LoadedPackages.Add(EditorWorldPackage);
+						break;
+					}
+				}
+			}
+		}
 
 		// Hot-reload the new packages...
 		UPackageTools::ReloadPackages(LoadedPackages);
@@ -1964,8 +1975,23 @@ void AssetViewUtils::SyncPathsFromSourceControl(const TArray<FString>& ContentPa
 		UPackageTools::UnloadPackages(PackageRawPtrsToUnload);
 
 		// Re-cache the SCC state...
-		SCCProvider.Execute(ISourceControlOperation::Create<FUpdateStatus>(), PathsOnDisk);
+		SCCProvider.Execute(ISourceControlOperation::Create<FUpdateStatus>(), AffectedFiles.Num() > 0 ? AffectedFiles : PathsOnDisk);
+
+		// Return result...
+		return (SyncResult == ECommandResult::Succeeded);
 	}
+
+	return true;
+}
+
+bool AssetViewUtils::SyncLatestFromSourceControl()
+{
+	return SyncPathsFromSourceControl(SourceControlHelpers::GetSourceControlLocations(), /*bCheckDependencies=*/false);
+}
+
+bool AssetViewUtils::SyncPathsFromSourceControl(const TArray<FString>& Paths)
+{
+	return SyncPathsFromSourceControl(Paths, /*bCheckDependencies=*/true);
 }
 
 void AssetViewUtils::ShowErrorNotifcation(const FText& InErrorMsg)
diff --git a/Engine/Source/Developer/AssetTools/Public/AssetViewUtils.h b/Engine/Source/Developer/AssetTools/Public/AssetViewUtils.h
--- a/Engine/Source/Developer/AssetTools/Public/AssetViewUtils.h
+++ b/Engine/Source/Developer/AssetTools/Public/AssetViewUtils.h
@@ -208,14 +208,18 @@ namespace AssetViewUtils
 	/** Gets the maximum path length for a cooked file. Changes behavior based on whether the editor experimental setting for long paths is enabled. */
 	ASSETTOOLS_API int32 GetMaxCookPathLen();
 
-	/** Syncs the specified packages from source control, other than any level assets which are currently being edited.*/
-	ASSETTOOLS_API void SyncPackagesFromSourceControl(const TArray<FString>& PackageNames, bool bIsSyncLatestOperation = false);
+	/** Syncs the specified packages from source control. */
+	UE_DEPRECATED(5.3, "The bIsSyncLatestOperation parameter is deprecated. Use SyncLatestFromSourceControl instead.")
+	ASSETTOOLS_API void SyncPackagesFromSourceControl(const TArray<FString>& PackageNames, bool bIsSyncLatestOperation);
 
-	/** Syncs the content from the specified paths from source control, other than any level assets which are currently being edited */
-	ASSETTOOLS_API void SyncPathsFromSourceControl(const TArray<FString>& ContentPaths);
+	/** Syncs the specified packages from source control. */
+	ASSETTOOLS_API bool SyncPackagesFromSourceControl(const TArray<FString>& PackageNames);
 
-	/** Syncs latest from source control and attempts to reload the world */
-	ASSETTOOLS_API void SyncLatestFromSourceControl();
+	/** Syncs the specified paths from source control. */
+	ASSETTOOLS_API bool SyncPathsFromSourceControl(const TArray<FString>& Paths);
+
+	/** Syncs latest from source control. */
+	ASSETTOOLS_API bool SyncLatestFromSourceControl();
 
 	/** Show an error notification toast if the given error message is not empty */
 	ASSETTOOLS_API void ShowErrorNotifcation(const FText& InErrorMsg);
diff --git a/Engine/Source/Editor/SourceControlWindows/Private/SourceControlWindows.cpp b/Engine/Source/Editor/SourceControlWindows/Private/SourceControlWindows.cpp
--- a/Engine/Source/Editor/SourceControlWindows/Private/SourceControlWindows.cpp
+++ b/Engine/Source/Editor/SourceControlWindows/Private/SourceControlWindows.cpp
@@ -176,7 +176,12 @@ bool FSourceControlWindows::SyncLatest()
 	// if properly saved or confirmation given, find all packages and use source control to update them.
 	if (bSaved)
 	{
-		AssetViewUtils::SyncLatestFromSourceControl();
+		bool bSuccess = AssetViewUtils::SyncLatestFromSourceControl();
+		if (!bSuccess)
+		{
+			FText Message(LOCTEXT("SCC_Sync_Failed", "Failed to sync files!"));
+			FMessageLog("SourceControl").Notify(Message);
+		}
 	}
 
 	return false;
