When moving an asset to a different folder in the Editor, do not show the "Check Out Assets" dialog if the Automatically Checkout on Asset Modification option is selected.
Moving an asset can replace the original asset with a redirector, so the source control system will consider the asset in the old location to be modified and the asset in the new location to be added.

diff --git a/Engine/Source/Developer/AssetTools/Private/AssetViewUtils.cpp b/Engine/Source/Developer/AssetTools/Private/AssetViewUtils.cpp
--- a/Engine/Source/Developer/AssetTools/Private/AssetViewUtils.cpp
+++ b/Engine/Source/Developer/AssetTools/Private/AssetViewUtils.cpp
@@ -18,6 +18,7 @@
 #include "Engine/Level.h"
 #include "Settings/ContentBrowserSettings.h"
 #include "Settings/EditorExperimentalSettings.h"
+#include "Settings/EditorLoadingSavingSettings.h"
 #include "SourceControlOperations.h"
 #include "ISourceControlModule.h"
 #include "SourceControlHelpers.h"
@@ -336,7 +337,8 @@ void AssetViewUtils::MoveAssets(const TArray<UObject*>& Assets, const FString& D
 
 	if ( AssetsAndNames.Num() > 0 )
 	{
-		AssetToolsModule.Get().RenameAssetsWithDialog(AssetsAndNames);
+		const UEditorLoadingSavingSettings* Settings = GetDefault<UEditorLoadingSavingSettings>();
+		AssetToolsModule.Get().RenameAssetsWithDialog(AssetsAndNames, Settings->GetAutomaticallyCheckoutOnAssetModification());
 	}
 }
 
