Unshelved from pending changelist '24180955':

[CrashFix] Removed invalid check()
- when I made assets diffable against null I must've missed this. Easy fix.

#lockout julien.marchand

diff --git a/Engine/Source/Developer/AssetTools/Public/AssetTypeActions_Base.h b/Engine/Source/Developer/AssetTools/Public/AssetTypeActions_Base.h
--- a/Engine/Source/Developer/AssetTools/Public/AssetTypeActions_Base.h
+++ b/Engine/Source/Developer/AssetTools/Public/AssetTypeActions_Base.h
@@ -96,8 +96,7 @@ public:
 
 	virtual void PerformAssetDiff(UObject* OldAsset, UObject* NewAsset, const struct FRevisionInfo& OldRevision, const struct FRevisionInfo& NewRevision) const override
 	{
-		check(OldAsset != nullptr);
-		check(NewAsset != nullptr);
+		check(OldAsset != nullptr || NewAsset != nullptr);
 
 		// Dump assets to temp text files
 		FString OldTextFilename = DumpAssetToTempFile(OldAsset);
