Performance improvement to opening Context Menu and other operations looking up AssetTypeActions

There is a host of functionality in the engine requiring looking up an AssetTypeAction object that best maps to a given class.  The metric of 'best' is the AssetTypeAction whose SupportedClass is the closest to the given class in a hierarchy.
The existing implementation searches through all AssetTypeActions (at time of writing, > 200) and queries each for their SupportedClass and then chooses the one whose SupportedClass is closest to the given class.

The improvement here exploits that UClass is the type information of UObjects, and UObjects don't support multiple-inheritance (ie. only one SuperClass).  Therefore we can walk the class hierarchy from the given Class up to find
the best AssetTypeAction.  This changes the linear search (with # of AssetTypeActions) to a nominal sub-linear search (with depth of inheritance).

To be able to do that, a new map was added that associates a registered AssetTypeAction's SupportedClass with the AssetTypeAction's index in it's array.  This lets us do a constant time lookup at the cost of a bit of memory.

There is further room for improvement with the usage of UFactory's public member functions such as `GetDisplayName`, `GetMenuCategories` and similar which still necessitates this query.  Opportunity exists in the NewAssetContextMenu and in the API of IAssetTools::GetNewAssetFactories
to generate a more performance oriented data structure.  There is also some opportunity to improve the construction of a FSoftClassPath from a UClass which can be fairly expensive. These all feed into the performance issue but ultimately this is the smallest change that provides the largest improvements at the lowest risk.

Testing:
Using debugger stepped through startup/shutdown of the editor.
Artificially added duplicate registrations to check codepaths for adding AssetTypeActions that with same supported class
Added unregisters to check codepaths for unregister.  Need to be careful with editor as AssetTool module is removed very early on outside of the usual lifetime of modules in EngineLoop::Exit so most AssetTypeActions do not actually get deregistered.

diff --git a/Engine/Source/Developer/AssetTools/Private/AssetTools.cpp b/Engine/Source/Developer/AssetTools/Private/AssetTools.cpp
--- a/Engine/Source/Developer/AssetTools/Private/AssetTools.cpp
+++ b/Engine/Source/Developer/AssetTools/Private/AssetTools.cpp
@@ -1337,8 +1337,40 @@ UAssetToolsImpl::UAssetToolsImpl(const FObjectInitializer& ObjectInitializer)
 	// coupled to your new system, and you should not create a dependency on your new system from AssetTools.
 }
 
+void UAssetToolsImpl::RemoveAssetTypeActionBySupportedClass(const UClass* SupportedClass)
+{
+	// Swap and pop the removed with the end, fixing up the lookup map if required
+	if (const int* IndexPtr = AssetTypeActionsLookup.Find(SupportedClass))
+	{
+		const int LastIndex = AssetTypeActionsList.Num() - 1;
+		const int Index = *IndexPtr;
+
+		AssetTypeActionsLookup.Remove(SupportedClass);
+		
+		if (LastIndex != Index)
+		{
+			TSharedRef<IAssetTypeActions>& Last = AssetTypeActionsList.Last();
+			AssetTypeActionsLookup.Add(Last->GetSupportedClass(), Index);
+			AssetTypeActionsList[Index] = MoveTemp(Last);
+		}
+
+		AssetTypeActionsList.Pop();
+	}
+
+	verify(AssetTypeActionsLookup.Num() == AssetTypeActionsList.Num());
+}
+
+
 void UAssetToolsImpl::RegisterAssetTypeActions(const TSharedRef<IAssetTypeActions>& NewActions)
 {
+	UClass* SupportedClass = NewActions->GetSupportedClass();
+
+	// If there is a duplicate AssetTypeActions registered, remove it
+	// only the most recent one will be used
+	RemoveAssetTypeActionBySupportedClass(SupportedClass);
+	
+	const int Index = AssetTypeActionsList.Num();
+	AssetTypeActionsLookup.Add(SupportedClass, Index);
 	AssetTypeActionsList.Add(NewActions);
 
 	if (!NewActions->IsAssetDefinitionInDisguise())
@@ -1378,8 +1410,9 @@ void UAssetToolsImpl::SyncAssetTypesToAssetDefinitions() const
 }
 
 void UAssetToolsImpl::UnregisterAssetTypeActions(const TSharedRef<IAssetTypeActions>& ActionsToRemove)
-{
-	AssetTypeActionsList.Remove(ActionsToRemove);
+{	
+	const UClass* SupportedClass = ActionsToRemove->GetSupportedClass();
+	RemoveAssetTypeActionBySupportedClass(SupportedClass);
 
 	// Now also remove the proxy AssetDefinition we registered for this type.
 	for (UAssetDefinition* AssetDefinition : UAssetDefinitionRegistry::Get()->GetAllAssetDefinitions())
@@ -1406,29 +1439,24 @@ void UAssetToolsImpl::GetAssetTypeActionsList( TArray<TWeakPtr<IAssetTypeActions
 
 TWeakPtr<IAssetTypeActions> UAssetToolsImpl::GetAssetTypeActionsForClass(const UClass* Class) const
 {
-	if (!Class)
-	{
-		return TSharedPtr<IAssetTypeActions>();
-	}
-
-	TSharedPtr<IAssetTypeActions> MostDerivedAssetTypeActions;
-
 	SyncAssetTypesToAssetDefinitions();
-	for (int32 TypeActionsIdx = 0; TypeActionsIdx < AssetTypeActionsList.Num(); ++TypeActionsIdx)
-	{
-		TSharedRef<IAssetTypeActions> TypeActions = AssetTypeActionsList[TypeActionsIdx];
-		UClass* SupportedClass = TypeActions->GetSupportedClass();
 
-		if ( Class->IsChildOf(SupportedClass) )
+	const UClass* CandidateClass = Class;
+	
+	while (CandidateClass)
+	{
+		const int* Index = AssetTypeActionsLookup.Find(CandidateClass);
+		if (Index)
 		{
-			if ( !MostDerivedAssetTypeActions.IsValid() || SupportedClass->IsChildOf( MostDerivedAssetTypeActions->GetSupportedClass() ) )
-			{
-				MostDerivedAssetTypeActions = TypeActions;
-			}
+			TWeakPtr<IAssetTypeActions> TypeActions = AssetTypeActionsList[*Index];
+			return TypeActions;
 		}
+
+		// Walk up the class hierarchy until we find the most derived ClassTypeAction
+		CandidateClass = CandidateClass->GetSuperClass();
 	}
 
-	return MostDerivedAssetTypeActions;
+	return nullptr;
 }
 
 bool UAssetToolsImpl::CanLocalize(const UClass* Class) const
diff --git a/Engine/Source/Developer/AssetTools/Private/AssetTools.h b/Engine/Source/Developer/AssetTools/Private/AssetTools.h
--- a/Engine/Source/Developer/AssetTools/Private/AssetTools.h
+++ b/Engine/Source/Developer/AssetTools/Private/AssetTools.h
@@ -221,10 +221,16 @@ private:
 	/** Make sure we're not syncing */
 	void SyncAssetTypesToAssetDefinitions() const;
 
+	/** Helper to remove an entry AssetTypeActionsList given a class returned by GetSupportedClass */
+	void RemoveAssetTypeActionBySupportedClass(const UClass* SupportedClass);
+
 private:
 	/** The list of all registered AssetTypeActions */
 	TArray<TSharedRef<IAssetTypeActions>> AssetTypeActionsList;
 
+	/** Lookup of the index into AssetTypeActionsList given the SupportedClass of a registered IAssetTypeActions */
+	TMap<const UClass*, int> AssetTypeActionsLookup;
+
 	/** The list of all registered ClassTypeActions */
 	TArray<TSharedRef<IClassTypeActions>> ClassTypeActionsList;
 
