SourceControl: Add start/finish delegates to AssetViewUtils sync operations.

#localization none
#tests manual

diff --git a/Engine/Source/Developer/AssetTools/Private/AssetViewUtils.cpp b/Engine/Source/Developer/AssetTools/Private/AssetViewUtils.cpp
--- a/Engine/Source/Developer/AssetTools/Private/AssetViewUtils.cpp
+++ b/Engine/Source/Developer/AssetTools/Private/AssetViewUtils.cpp
@@ -58,6 +58,12 @@ namespace AssetViewUtils
 	/** Callback used when a folder is moved or renamed */
 	static FOnFolderPathChanged OnFolderPathChangedDelegate;
 
+	/** Callback used when a sync starts */
+	static FOnSyncStart FOnSyncStartDelegate;
+
+	/** Callback used when a sync finishes */
+	static FOnSyncFinish FOnSyncFinishDelegate;
+
 	/** Keep a map of all the paths that have custom colors, so updating the color in one location updates them all */
 	static TMap< FString, FLinearColor > PathColors;
 
@@ -1721,6 +1727,9 @@ bool AssetViewUtils::SyncPackagesFromSourceControl(const TArray<FString>& Packag
 {
 	if (PackageNames.Num() > 0)
 	{
+		// Broadcast sync starting...
+		AssetViewUtils::OnSyncStart().Broadcast();
+
 		TArray<FString> PackageNamesToSync = PackageNames;
 
 		// Warn about any packages that are being synced without also getting the newest version of their dependencies...
@@ -1792,6 +1801,9 @@ bool AssetViewUtils::SyncPackagesFromSourceControl(const TArray<FString>& Packag
 		// Re-cache the SCC state...
 		SCCProvider.Execute(ISourceControlOperation::Create<FUpdateStatus>(), PackageFilenames);
 
+		// Broadcast sync finished...
+		AssetViewUtils::OnSyncFinish().Broadcast(SyncResult == ECommandResult::Succeeded, &PackageFilenames);
+
 		// Return result...
 		return (SyncResult == ECommandResult::Succeeded);
 	}
@@ -1824,6 +1836,9 @@ static bool SyncPathsFromSourceControl(const TArray<FString>& Paths, bool bCheck
 
 	if (PathsOnDisk.Num() > 0)
 	{
+		// Broadcast sync starting...
+		AssetViewUtils::OnSyncStart().Broadcast();
+
 		ISourceControlProvider& SCCProvider = ISourceControlModule::Get().GetProvider();
 
 		// Get all the assets about to be synced in those path(s) on disk...
@@ -1979,6 +1994,9 @@ static bool SyncPathsFromSourceControl(const TArray<FString>& Paths, bool bCheck
 		// Re-cache the SCC state...
 		SCCProvider.Execute(ISourceControlOperation::Create<FUpdateStatus>(), AffectedFiles.Num() > 0 ? AffectedFiles : PathsOnDisk);
 
+		// Broadcast sync finished...
+		AssetViewUtils::OnSyncFinish().Broadcast(SyncResult == ECommandResult::Succeeded, AffectedFiles.Num() > 0 ? &AffectedFiles : nullptr);
+
 		// Return result...
 		return (SyncResult == ECommandResult::Succeeded);
 	}
@@ -2020,4 +2038,14 @@ AssetViewUtils::FOnFolderPathChanged& AssetViewUtils::OnFolderPathChanged()
 	return OnFolderPathChangedDelegate;
 }
 
+AssetViewUtils::FOnSyncStart& AssetViewUtils::OnSyncStart()
+{
+	return FOnSyncStartDelegate;
+}
+
+AssetViewUtils::FOnSyncFinish& AssetViewUtils::OnSyncFinish()
+{
+	return FOnSyncFinishDelegate;
+}
+
 #undef LOCTEXT_NAMESPACE
diff --git a/Engine/Source/Developer/AssetTools/Public/AssetViewUtils.h b/Engine/Source/Developer/AssetTools/Public/AssetViewUtils.h
--- a/Engine/Source/Developer/AssetTools/Public/AssetViewUtils.h
+++ b/Engine/Source/Developer/AssetTools/Public/AssetViewUtils.h
@@ -232,4 +232,16 @@ namespace AssetViewUtils
 	using FMovedContentFolder = TTuple<FString /*OldPath*/, FString /*NewPath*/>;
 	DECLARE_MULTICAST_DELEGATE_OneParam(FOnFolderPathChanged, TArrayView<const FMovedContentFolder> /*ChangedPaths*/);
 	ASSETTOOLS_API FOnFolderPathChanged& OnFolderPathChanged();
+
+	/** Called when a sync from AssetViewUtils starts */
+	DECLARE_MULTICAST_DELEGATE(FOnSyncStart);
+	ASSETTOOLS_API FOnSyncStart& OnSyncStart();
+
+	/** 
+	 * Called when a sync from AssetViewUtils finishes
+	 * If Files is non-nullptr, those are the files affected by the sync.
+	 * If Files is nullptr, assume any file might have been changed by the sync.
+	 */
+	DECLARE_MULTICAST_DELEGATE_TwoParams(FOnSyncFinish, bool /*bSucceeded*/, const TArray<FString>* /*Files*/);
+	ASSETTOOLS_API FOnSyncFinish& OnSyncFinish();
 }
