Set the asset filter in UAssetToolsImpl::CreateAssetWithDialog

This will use the Factory supported class as a filter for assets in the Save Asset As dialog. Before this change anyone using this function would get a dialog that displayed all assets in the project.

This also makes it consistent with how Save Level As works, which only shows Levels in its dialog

diff --git a/Engine/Source/Developer/AssetTools/Private/AssetTools.cpp b/Engine/Source/Developer/AssetTools/Private/AssetTools.cpp
--- a/Engine/Source/Developer/AssetTools/Private/AssetTools.cpp
+++ b/Engine/Source/Developer/AssetTools/Private/AssetTools.cpp
@@ -1738,6 +1738,7 @@ UObject* UAssetToolsImpl::CreateAssetWithDialog(const FString& AssetName, const
 		SaveAssetDialogConfig.DefaultPath = PackagePath;
 		SaveAssetDialogConfig.DefaultAssetName = AssetName;
 		SaveAssetDialogConfig.ExistingAssetPolicy = ESaveAssetDialogExistingAssetPolicy::AllowButWarn;
+		SaveAssetDialogConfig.AssetClassNames.Add(Factory->GetSupportedClass()->StaticClass()->GetClassPathName());
 
 		FContentBrowserModule& ContentBrowserModule = FModuleManager::LoadModuleChecked<FContentBrowserModule>("ContentBrowser");
 		FString SaveObjectPath = ContentBrowserModule.Get().CreateModalSaveAssetDialog(SaveAssetDialogConfig);
