Fix the asset filter in UAssetToolsImpl::CreateAssetWithDialog to use the Factory supported class path

Fix a mistake made in CL24423110 where the class path name added to the filter was from a UClass instead of being from the supported class that the filter is supposed to use

diff --git a/Engine/Source/Developer/AssetTools/Private/AssetTools.cpp b/Engine/Source/Developer/AssetTools/Private/AssetTools.cpp
--- a/Engine/Source/Developer/AssetTools/Private/AssetTools.cpp
+++ b/Engine/Source/Developer/AssetTools/Private/AssetTools.cpp
@@ -1765,7 +1765,7 @@ UObject* UAssetToolsImpl::CreateAssetWithDialog(const FString& AssetName, const
 		SaveAssetDialogConfig.DefaultPath = PackagePath;
 		SaveAssetDialogConfig.DefaultAssetName = AssetName;
 		SaveAssetDialogConfig.ExistingAssetPolicy = ESaveAssetDialogExistingAssetPolicy::AllowButWarn;
-		SaveAssetDialogConfig.AssetClassNames.Add(Factory->GetSupportedClass()->StaticClass()->GetClassPathName());
+		SaveAssetDialogConfig.AssetClassNames.Add(Factory->GetSupportedClass()->GetClassPathName());
 
 		FContentBrowserModule& ContentBrowserModule = FModuleManager::LoadModuleChecked<FContentBrowserModule>("ContentBrowser");
 		FString SaveObjectPath = ContentBrowserModule.Get().CreateModalSaveAssetDialog(SaveAssetDialogConfig);
