Fixed advanced copy of partitioned world.

diff --git a/Engine/Source/Developer/AssetTools/Private/AssetTools.cpp b/Engine/Source/Developer/AssetTools/Private/AssetTools.cpp
--- a/Engine/Source/Developer/AssetTools/Private/AssetTools.cpp
+++ b/Engine/Source/Developer/AssetTools/Private/AssetTools.cpp
@@ -2329,6 +2329,14 @@ bool UAssetToolsImpl::AdvancedCopyPackages(
 			}
 		}
 
+		for (FName CopiedWorldPartitionMap : CopiedWorldPartitionMaps)
+		{
+			UPackage* CopiedWorldPartitionPackage = FindPackage(nullptr, *CopiedWorldPartitionMap.ToString());
+			UWorld* World = UWorld::FindWorldInPackage(CopiedWorldPartitionPackage);
+			UWorldPartition* WorldPartition = World->GetWorldPartition();
+			WorldPartition->Uninitialize();
+		}
+
 		FMessageLog AdvancedCopyLog("AssetTools");
 		FText LogMessage = FText::FromString(TEXT("Advanced content copy completed successfully!"));
 		EMessageSeverity::Type Severity = EMessageSeverity::Info;
