Updated advanced copy to use the lazy consolidation method by default

#rnx

diff --git a/Engine/Source/Developer/AssetTools/Private/AssetTools.cpp b/Engine/Source/Developer/AssetTools/Private/AssetTools.cpp
--- a/Engine/Source/Developer/AssetTools/Private/AssetTools.cpp
+++ b/Engine/Source/Developer/AssetTools/Private/AssetTools.cpp
@@ -421,19 +421,19 @@ namespace UE::AssetTools::Private
 
 	enum class EAdvancedCopyConsolidationMethod : uint8
 	{
-		/** Consolidate everything once (new, fast) */
+		/** Consolidate everything once (default, fast) */
 		Lazy = 0,
-		/** Consolidate once per-asset (default, slow) */
+		/** Consolidate once per-asset (legacy, slow) */
 		PerAsset = 1,
 		/** Consolidate once per-asset dependency (legacy, glacial) */
 		PerAssetDependency = 2,
 	};
 
-	int32 AdvancedCopyConsolidationMethodInt = static_cast<int32>(EAdvancedCopyConsolidationMethod::PerAsset);
+	int32 AdvancedCopyConsolidationMethodInt = static_cast<int32>(EAdvancedCopyConsolidationMethod::Lazy);
 	FAutoConsoleVariableRef CVarAdvancedCopyConsolidationMethod(
 		TEXT("AssetTools.AdvancedCopyConsolidationMethod"),
 		AdvancedCopyConsolidationMethodInt,
-		TEXT("0: Lazy (new, fast), 1: PerAsset (default, slow), 2: PerAssetDependency (legacy, glacial)")
+		TEXT("0: Lazy (default, fast), 1: PerAsset (legacy, slow), 2: PerAssetDependency (legacy, glacial)")
 	);
 
 	EAdvancedCopyConsolidationMethod GetAdvancedCopyConsolidationMethod()
