Remove "All Types" as an export file type as it doesn't really make sense.

diff --git a/Engine/Source/Developer/AssetTools/Private/AssetTools.cpp b/Engine/Source/Developer/AssetTools/Private/AssetTools.cpp
--- a/Engine/Source/Developer/AssetTools/Private/AssetTools.cpp
+++ b/Engine/Source/Developer/AssetTools/Private/AssetTools.cpp
@@ -3936,7 +3936,6 @@ void UAssetToolsImpl::ExportAssetsInternal(const TArray<UObject*>& ObjectsToExpo
 			}
 			Extensions += AllExtensions[ExtIndex];
 		}
-		FileTypes = FString::Printf(TEXT("%s|All Files (%s)|%s"), *FileTypes, *Extensions, *Extensions);
 
 		FString SaveFileName;
 		if (bPromptIndividualFilenames)
