Addressing a bug with AdvancedCopy, where two asset dependencies with the same name (from different plugins) would flatten down to overwrite eachother (causing a crash).

NOTE: This could significantly lengthen advanced copy destination paths, potentially blowing past the path character limit normally enforced by UE. In it's current state, the copy will just proceed normally -- seemingly working. This is probably something to investigate more, but doesn't appear to be a high-pri issue atm. A JIRA has been logged to capture this though: UE-182280

diff --git a/Engine/Source/Developer/AssetTools/Private/AssetTools.cpp b/Engine/Source/Developer/AssetTools/Private/AssetTools.cpp
--- a/Engine/Source/Developer/AssetTools/Private/AssetTools.cpp
+++ b/Engine/Source/Developer/AssetTools/Private/AssetTools.cpp
@@ -1999,7 +1999,22 @@ FString UAssetToolsImpl::GenerateAdvancedCopyDestinationPackageName(const FStrin
 	{
 		// Use the passed in default path
 		// Normal path
-		DestinationPackageName = DestinationFolder + "/" + FPackageName::GetShortName(SourcePackage);
+		// 
+		// NOTE: We cannot use a shortened path from SourcePackage because there may be assets that share names.
+		//       To prevent assets from copying over eachother we need a unique destination path.
+
+		// A little nicety: don't expose '/Game/' to users -- instead create a subfolder using the project's name
+		if (SourcePackage.StartsWith(TEXT("/Game/")))
+		{
+			FString ThrowawayRoot, SubPath, PkgName;
+			FPackageName::SplitLongPackageName(SourcePackage, ThrowawayRoot, SubPath, PkgName);
+
+			DestinationPackageName = DestinationFolder / FApp::GetProjectName() / SubPath / PkgName;
+		}
+		else // Other non-'/Game/' paths (ones from plugins, etc.)...
+		{
+			DestinationPackageName = DestinationFolder / SourcePackage;
+		}
 	}
 	else
 	{
