[Asset Definition] Fix an issue where an invalid FName could cause category path look up to assert

diff --git a/Engine/Source/Developer/AssetTools/Private/AssetTools.cpp b/Engine/Source/Developer/AssetTools/Private/AssetTools.cpp
--- a/Engine/Source/Developer/AssetTools/Private/AssetTools.cpp
+++ b/Engine/Source/Developer/AssetTools/Private/AssetTools.cpp
@@ -166,18 +166,24 @@ public:
 		uint32 CategoryBits = 0;
 		for (const FAssetCategoryPath& Category : AssetDefinitionPtr.Get()->GetAssetCategories())
 		{
-			if (Category.GetCategory() == EAssetCategoryPaths::Basic.GetCategory())
+			const FName& CategoryName = Category.GetCategory();
+			if (!ensureMsgf(CategoryName.IsValid(),	TEXT("Found an invalid category path for Asset Definition: %s"), *(AssetDefinitionPtr.Get()->GetName())))
+			{
+				return CategoryBits;
+			}
+
+			if (CategoryName == EAssetCategoryPaths::Basic.GetCategory())
 			{
 				CategoryBits |= EAssetTypeCategories::Basic;
 			}
 			else
 			{
-				if (Category.GetCategory() != EAssetCategoryPaths::Misc.GetCategory())
+				if (CategoryName != EAssetCategoryPaths::Misc.GetCategory())
 				{
-					const EAssetTypeCategories::Type AdvancedCategoryBit = AssetTools.FindAdvancedAssetCategory(Category.GetCategory());
+					const EAssetTypeCategories::Type AdvancedCategoryBit = AssetTools.FindAdvancedAssetCategory(CategoryName);
 					if (AdvancedCategoryBit == EAssetTypeCategories::Misc)
 					{
-						CategoryBits |= AssetTools.RegisterAdvancedAssetCategory(Category.GetCategory(), Category.GetCategoryText());
+						CategoryBits |= AssetTools.RegisterAdvancedAssetCategory(CategoryName, Category.GetCategoryText());
 					}
 					else
 					{
