Add export class permission list
#rnx

diff --git a/Engine/Plugins/Editor/ContentBrowser/ContentBrowserAssetDataSource/Source/ContentBrowserAssetDataSource/Private/AssetFileContextMenu.cpp b/Engine/Plugins/Editor/ContentBrowser/ContentBrowserAssetDataSource/Source/ContentBrowserAssetDataSource/Private/AssetFileContextMenu.cpp
--- a/Engine/Plugins/Editor/ContentBrowser/ContentBrowserAssetDataSource/Source/ContentBrowserAssetDataSource/Private/AssetFileContextMenu.cpp
+++ b/Engine/Plugins/Editor/ContentBrowser/ContentBrowserAssetDataSource/Source/ContentBrowserAssetDataSource/Private/AssetFileContextMenu.cpp
@@ -377,10 +377,7 @@ void FAssetFileContextMenu::MakeAssetActionsSubMenu(UToolMenu* Menu)
 	if (bCanBeModified)
 	{
 		FToolMenuSection& Section = Menu->AddSection("AssetContextMoveActions", LOCTEXT("AssetContextMoveActionsMenuHeading", "Move"));
-		const bool bCanExportAllSelectedAssets = !Algo::AnyOf(SelectedAssets, [](const FAssetData& AssetData)
-		{
-			return AssetData.HasAnyPackageFlags(EPackageFlags::PKG_DisallowExport);
-		});
+		const bool bCanExportAllSelectedAssets = FAssetToolsModule::GetModule().Get().CanExportAssets(SelectedAssets);
 
 		if (bCanExportAllSelectedAssets)
 		{
diff --git a/Engine/Source/Developer/AssetTools/Private/AssetTools.cpp b/Engine/Source/Developer/AssetTools/Private/AssetTools.cpp
--- a/Engine/Source/Developer/AssetTools/Private/AssetTools.cpp
+++ b/Engine/Source/Developer/AssetTools/Private/AssetTools.cpp
@@ -3841,6 +3841,13 @@ void UAssetToolsImpl::ExportAssetsInternal(const TArray<UObject*>& ObjectsToExpo
 
 	// Export the objects.
 	bool bAnyObjectMissingSourceData = false;
+
+	// Permission filter based on class
+	TSharedPtr<FPathPermissionList> ExportClassPermissionList = GetAssetClassPathPermissionList(EAssetClassAction::ExportAsset);
+	const bool bExportClassPermissionListHasFiltering = ExportClassPermissionList && ExportClassPermissionList->HasFiltering();
+	// Use TMap as optimization because FPathPermissionList performs slow TArray<FString> search
+	TMap<UClass*, bool> ClassPassesFilterResults;
+
 	for (int32 Index = 0; Index < ObjectsToExport.Num(); Index++)
 	{
 		GWarn->StatusUpdate(Index, ObjectsToExport.Num(), FText::Format(NSLOCTEXT("UnrealEd", "Exportingf", "Exporting ({0} of {1})"), FText::AsNumber(Index), FText::AsNumber(ObjectsToExport.Num())));
@@ -3856,6 +3863,34 @@ void UAssetToolsImpl::ExportAssetsInternal(const TArray<UObject*>& ObjectsToExpo
 			continue;
 		}
 
+		// Look at export class permission filter
+		if (bExportClassPermissionListHasFiltering)
+		{
+			// Reuse results as optimization
+			const bool* FoundPassesFilterResult = ClassPassesFilterResults.Find(ObjectToExport->GetClass());
+			if (FoundPassesFilterResult)
+			{
+				if (*FoundPassesFilterResult == false)
+				{
+					continue;
+				}
+			}
+			else
+			{
+				FNameBuilder AssetClassPath;
+				ObjectToExport->GetClass()->GetPathName(nullptr, AssetClassPath);
+				const bool bPassesFilterResult = ExportClassPermissionList->PassesFilter(AssetClassPath.ToView());
+
+				// Save result to be reused as optimization
+				ClassPassesFilterResults.Add(ObjectToExport->GetClass(), bPassesFilterResult);
+
+				if (bPassesFilterResult == false)
+				{
+					continue;
+				}
+			}
+		}
+
 		// Find all the exporters that can export this type of object and construct an export file dialog.
 		TArray<FString> AllFileTypes;
 		TArray<FString> AllExtensions;
@@ -4136,6 +4171,47 @@ void UAssetToolsImpl::ExportAssetsInternal(const TArray<UObject*>& ObjectsToExpo
 	FEditorDirectories::Get().SetLastDirectory(ELastDirectory::GENERIC_EXPORT, LastExportPath);
 }
 
+bool UAssetToolsImpl::CanExportAssets(const TArray<FAssetData>& AssetsToExport) const
+{
+	const TSharedPtr<FPathPermissionList> PermissionList = GetAssetClassPathPermissionList(EAssetClassAction::ExportAsset);
+	if (PermissionList && PermissionList->HasFiltering())
+	{
+		// Use TSet as optimization because FPathPermissionList performs slow TArray<FString> search
+		TSet<FTopLevelAssetPath> AlreadyPassed;
+		for (const FAssetData& AssetData : AssetsToExport)
+		{
+			if (AssetData.HasAnyPackageFlags(EPackageFlags::PKG_DisallowExport))
+			{
+				return false;
+			}
+
+			if (!AlreadyPassed.Contains(AssetData.AssetClassPath))
+			{
+				FNameBuilder ClassPath;
+				AssetData.AssetClassPath.AppendString(ClassPath);
+				if (!PermissionList->PassesFilter(ClassPath.ToView()))
+				{
+					return false;
+				}
+
+				AlreadyPassed.Add(AssetData.AssetClassPath);
+			}
+		}
+	}
+	else
+	{
+		for (const FAssetData& AssetData : AssetsToExport)
+		{
+			if (AssetData.HasAnyPackageFlags(EPackageFlags::PKG_DisallowExport))
+			{
+				return false;
+			}
+		}
+	}
+
+	return true;
+}
+
 UAssetToolsImpl& UAssetToolsImpl::Get()
 {
 	FAssetToolsModule& Module = FModuleManager::GetModuleChecked<FAssetToolsModule>("AssetTools");
diff --git a/Engine/Source/Developer/AssetTools/Private/AssetTools.h b/Engine/Source/Developer/AssetTools/Private/AssetTools.h
--- a/Engine/Source/Developer/AssetTools/Private/AssetTools.h
+++ b/Engine/Source/Developer/AssetTools/Private/AssetTools.h
@@ -103,6 +103,7 @@ public:
 	virtual void ExportAssets(const TArray<UObject*>& AssetsToExport, const FString& ExportPath) const override;
 	virtual void ExportAssetsWithDialog(const TArray<UObject*>& AssetsToExport, bool bPromptForIndividualFilenames) override;
 	virtual void ExportAssetsWithDialog(const TArray<FString>& AssetsToExport, bool bPromptForIndividualFilenames) override;
+	virtual bool CanExportAssets(const TArray<FAssetData>& AssetsToExport) const override;
 	virtual void CreateUniqueAssetName(const FString& InBasePackageName, const FString& InSuffix, FString& OutPackageName, FString& OutAssetName) override;
 	virtual bool AssetUsesGenericThumbnail( const FAssetData& AssetData ) const override;
 	virtual void DiffAgainstDepot(UObject* InObject, const FString& InPackagePath, const FString& InPackageName) const override;
diff --git a/Engine/Source/Developer/AssetTools/Public/IAssetTools.h b/Engine/Source/Developer/AssetTools/Public/IAssetTools.h
--- a/Engine/Source/Developer/AssetTools/Public/IAssetTools.h
+++ b/Engine/Source/Developer/AssetTools/Public/IAssetTools.h
@@ -46,6 +46,8 @@ enum class EAssetClassAction : uint8
 	ViewAsset,
 	/** Whether an asset can be imported */
 	ImportAsset,
+	/** Whether an asset can be exported */
+	ExportAsset,
 	AllAssetActions
 };
 
@@ -459,6 +461,14 @@ public:
 	 */
 	virtual void ExportAssetsWithDialog(const TArray<UObject*>& AssetsToExport, bool bPromptForIndividualFilenames) = 0;
 
+	/**
+	 * Check if specified assets can be exported.
+	 *
+	 * @param	AssetsToExport					List of assets to export
+	 * @return true if all assets specified can be exported
+	 */
+	virtual bool CanExportAssets(const TArray<FAssetData>& AssetsToExport) const = 0;
+
 	/** Creates a unique package and asset name taking the form InBasePackageName+InSuffix */
 	UFUNCTION(BlueprintCallable, Category = "Editor Scripting | Asset Tools")
 	virtual void CreateUniqueAssetName(const FString& InBasePackageName, const FString& InSuffix, FString& OutPackageName, FString& OutAssetName) = 0;
