[EIS] Minor improvement to the error messaging when renaming an asset.

diff --git a/Engine/Source/Developer/AssetTools/Private/AssetViewUtils.cpp b/Engine/Source/Developer/AssetTools/Private/AssetViewUtils.cpp
--- a/Engine/Source/Developer/AssetTools/Private/AssetViewUtils.cpp
+++ b/Engine/Source/Developer/AssetTools/Private/AssetViewUtils.cpp
@@ -1315,17 +1315,17 @@ bool AssetViewUtils::IsValidObjectPathForCreate(const FString& ObjectPath, const
 {
 	const FString ObjectName = FPackageName::ObjectPathToPathWithinPackage(ObjectPath);
 
-	// Make sure the name is not already a class or otherwise invalid for saving
-	if ( !FFileHelper::IsFilenameValidForSaving(ObjectName, OutErrorMessage) )
+	// Make sure the new name only contains valid characters
+	if (!FName::IsValidXName( ObjectName, INVALID_OBJECTNAME_CHARACTERS INVALID_LONGPACKAGE_CHARACTERS, &OutErrorMessage))
 	{
 		// Return false to indicate that the user should enter a new name
 		return false;
 	}
 
-	// Make sure the new name only contains valid characters
-	if ( !FName::IsValidXName( ObjectName, INVALID_OBJECTNAME_CHARACTERS INVALID_LONGPACKAGE_CHARACTERS, &OutErrorMessage ) )
+	const FString PackageName = FPackageName::ObjectPathToPackageName(ObjectPath);
+
+	if (!IsValidPackageForCooking(PackageName, OutErrorMessage))
 	{
-		// Return false to indicate that the user should enter a new name
 		return false;
 	}
 
@@ -1337,13 +1337,6 @@ bool AssetViewUtils::IsValidObjectPathForCreate(const FString& ObjectPath, const
 		return false;
 	}
 
-	const FString PackageName = FPackageName::ObjectPathToPackageName(ObjectPath);
-
-	if (!IsValidPackageForCooking(PackageName, OutErrorMessage))
-	{
-		return false;
-	}
-
 	// Make sure we are not creating an path that is too long for the OS
 	FString RelativePathFilename;
 	if (!FPackageName::TryConvertLongPackageNameToFilename(PackageName, RelativePathFilename, FPackageName::GetAssetPackageExtension()))	// full relative path with name + extension
@@ -1401,6 +1394,19 @@ bool AssetViewUtils::IsValidObjectPathForCreate(const FString& ObjectPath, const
 		return false;
 	}
 
+	/** 
+	 * Make sure the name is not already a class or otherwise invalid for saving
+	 * Some other function above also test for some issues that this function can also find,
+	 * but they have an better error message more suited for the asset view.
+	 * Because of this, this test should be the last one to be run.
+	 */
+	if ( !FFileHelper::IsFilenameValidForSaving(ObjectName, OutErrorMessage) )
+	{
+		// Return false to indicate that the user should enter a new name
+		return false;
+	}
+
+
 	return true;
 }
 
