Remove the consolidation method CVar now that lazy consolidation has had time to bed in as the default

#rnx

diff --git a/Engine/Source/Developer/AssetTools/Private/AssetTools.cpp b/Engine/Source/Developer/AssetTools/Private/AssetTools.cpp
--- a/Engine/Source/Developer/AssetTools/Private/AssetTools.cpp
+++ b/Engine/Source/Developer/AssetTools/Private/AssetTools.cpp
@@ -425,32 +425,6 @@ namespace UE::AssetTools::Private
 		TEXT("When set, The package migration will use the new implementation made for 5.1.")
 	);
 
-	enum class EAdvancedCopyConsolidationMethod : uint8
-	{
-		/** Consolidate everything once (default, fast) */
-		Lazy = 0,
-		/** Consolidate once per-asset (legacy, slow) */
-		PerAsset = 1,
-		/** Consolidate once per-asset dependency (legacy, glacial) */
-		PerAssetDependency = 2,
-	};
-
-	int32 AdvancedCopyConsolidationMethodInt = static_cast<int32>(EAdvancedCopyConsolidationMethod::Lazy);
-	FAutoConsoleVariableRef CVarAdvancedCopyConsolidationMethod(
-		TEXT("AssetTools.AdvancedCopyConsolidationMethod"),
-		AdvancedCopyConsolidationMethodInt,
-		TEXT("0: Lazy (default, fast), 1: PerAsset (legacy, slow), 2: PerAssetDependency (legacy, glacial)")
-	);
-
-	EAdvancedCopyConsolidationMethod GetAdvancedCopyConsolidationMethod()
-	{
-		if (AdvancedCopyConsolidationMethodInt >= static_cast<int32>(EAdvancedCopyConsolidationMethod::Lazy) && AdvancedCopyConsolidationMethodInt <= static_cast<int32>(EAdvancedCopyConsolidationMethod::PerAssetDependency))
-		{
-			return static_cast<EAdvancedCopyConsolidationMethod>(AdvancedCopyConsolidationMethodInt);
-		}
-		return EAdvancedCopyConsolidationMethod::PerAsset;
-	}
-
 	bool bFollowRedirectorsWhenImporting = false;
 	static FAutoConsoleVariableRef CVarFollowRedirectorsWhenImporting(
 		TEXT("AssetTools.FollowRedirectorsWhenImporting"),
@@ -2298,26 +2272,12 @@ bool UAssetToolsImpl::AdvancedCopyPackages(
 						}
 					}
 				}
-
-				if (UE::AssetTools::Private::GetAdvancedCopyConsolidationMethod() == UE::AssetTools::Private::EAdvancedCopyConsolidationMethod::PerAssetDependency)
-				{
-					ConsolidatePendingObjects();
-				}
-			}
-
-			if (UE::AssetTools::Private::GetAdvancedCopyConsolidationMethod() == UE::AssetTools::Private::EAdvancedCopyConsolidationMethod::PerAsset)
-			{
-				ConsolidatePendingObjects();
 			}
 		}
 
 		LoopProgress.Reset();
 
-		if (UE::AssetTools::Private::GetAdvancedCopyConsolidationMethod() == UE::AssetTools::Private::EAdvancedCopyConsolidationMethod::Lazy)
-		{
-			ConsolidatePendingObjects();
-		}
-
+		ConsolidatePendingObjects();
 		check(Consolidations.Num() == 0);
 
 		ObjectTools::CompileBlueprintsAfterRefUpdate(NewObjectSet.Array());
