Removed unnecessary private include modules.

diff --git a/Engine/Plugins/AI/EnvironmentQueryEditor/Source/EnvironmentQueryEditor/EnvironmentQueryEditor.Build.cs b/Engine/Plugins/AI/EnvironmentQueryEditor/Source/EnvironmentQueryEditor/EnvironmentQueryEditor.Build.cs
--- a/Engine/Plugins/AI/EnvironmentQueryEditor/Source/EnvironmentQueryEditor/EnvironmentQueryEditor.Build.cs
+++ b/Engine/Plugins/AI/EnvironmentQueryEditor/Source/EnvironmentQueryEditor/EnvironmentQueryEditor.Build.cs
@@ -19,9 +19,7 @@ public class EnvironmentQueryEditor : ModuleRules
 		PrivateIncludePathModuleNames.AddRange(
 		   new string[] {
 				"AssetRegistry",
-				"AssetTools",
 				"DesktopPlatform",
-				"LevelEditor",
 		   });
 
 		PrivateDependencyModuleNames.AddRange(
diff --git a/Engine/Plugins/AudioGameplayVolume/Source/AudioGameplayVolumeEditor/AudioGameplayVolumeEditor.Build.cs b/Engine/Plugins/AudioGameplayVolume/Source/AudioGameplayVolumeEditor/AudioGameplayVolumeEditor.Build.cs
--- a/Engine/Plugins/AudioGameplayVolume/Source/AudioGameplayVolumeEditor/AudioGameplayVolumeEditor.Build.cs
+++ b/Engine/Plugins/AudioGameplayVolume/Source/AudioGameplayVolumeEditor/AudioGameplayVolumeEditor.Build.cs
@@ -31,7 +31,6 @@ public class AudioGameplayVolumeEditor : ModuleRules
 		PrivateIncludePathModuleNames.AddRange(
 			new string[]
 			{
-				"AssetTools",
 			}
 			);
 	}
diff --git a/Engine/Plugins/Cameras/CameraShakePreviewer/Source/CameraShakePreviewer/CameraShakePreviewer.Build.cs b/Engine/Plugins/Cameras/CameraShakePreviewer/Source/CameraShakePreviewer/CameraShakePreviewer.Build.cs
--- a/Engine/Plugins/Cameras/CameraShakePreviewer/Source/CameraShakePreviewer/CameraShakePreviewer.Build.cs
+++ b/Engine/Plugins/Cameras/CameraShakePreviewer/Source/CameraShakePreviewer/CameraShakePreviewer.Build.cs
@@ -8,9 +8,6 @@ public class CameraShakePreviewer : ModuleRules
 	{
 		PrivateIncludePathModuleNames.AddRange(
 			new string[] {
-				"AssetTools",
-				"AssetRegistry",
-				"ContentBrowser",
 				"WorkspaceMenuStructure",
 				}
 		);
diff --git a/Engine/Plugins/Editor/DataValidation/Source/DataValidation/DataValidation.Build.cs b/Engine/Plugins/Editor/DataValidation/Source/DataValidation/DataValidation.Build.cs
--- a/Engine/Plugins/Editor/DataValidation/Source/DataValidation/DataValidation.Build.cs
+++ b/Engine/Plugins/Editor/DataValidation/Source/DataValidation/DataValidation.Build.cs
@@ -8,7 +8,6 @@ public class DataValidation : ModuleRules
 	{
 		PrivateIncludePathModuleNames.AddRange(
 			new string[] {
-				"MessageLog",
 			}
 		);
 
diff --git a/Engine/Plugins/Experimental/Mutable/Source/CustomizableObjectPopulation/CustomizableObjectPopulation.Build.cs b/Engine/Plugins/Experimental/Mutable/Source/CustomizableObjectPopulation/CustomizableObjectPopulation.Build.cs
--- a/Engine/Plugins/Experimental/Mutable/Source/CustomizableObjectPopulation/CustomizableObjectPopulation.Build.cs
+++ b/Engine/Plugins/Experimental/Mutable/Source/CustomizableObjectPopulation/CustomizableObjectPopulation.Build.cs
@@ -32,7 +32,6 @@ public class CustomizableObjectPopulation : ModuleRules
 
         PrivateIncludePathModuleNames.AddRange(
         new string[] {
-				"TargetPlatform"
 		});
 
         if (TargetRules.bBuildEditor == true)
diff --git a/Engine/Plugins/Experimental/Mutable/Source/CustomizableObjectPopulationEditor/CustomizableObjectPopulationEditor.Build.cs b/Engine/Plugins/Experimental/Mutable/Source/CustomizableObjectPopulationEditor/CustomizableObjectPopulationEditor.Build.cs
--- a/Engine/Plugins/Experimental/Mutable/Source/CustomizableObjectPopulationEditor/CustomizableObjectPopulationEditor.Build.cs
+++ b/Engine/Plugins/Experimental/Mutable/Source/CustomizableObjectPopulationEditor/CustomizableObjectPopulationEditor.Build.cs
@@ -15,14 +15,8 @@ public class CustomizableObjectPopulationEditor : ModuleRules
 		PrivateIncludePathModuleNames.AddRange(
 			new string[] {
 				"AssetRegistry",
-				"EditorWidgets",
-				"MeshUtilities",
-				"SkeletonEditor",
-				"Persona",
 				"WorkspaceMenuStructure",
 
-				"MessageLog",
-				"KismetWidgets",
 			}
 		);
 
diff --git a/Engine/Plugins/FX/Niagara/Source/NiagaraEditor/NiagaraEditor.Build.cs b/Engine/Plugins/FX/Niagara/Source/NiagaraEditor/NiagaraEditor.Build.cs
--- a/Engine/Plugins/FX/Niagara/Source/NiagaraEditor/NiagaraEditor.Build.cs
+++ b/Engine/Plugins/FX/Niagara/Source/NiagaraEditor/NiagaraEditor.Build.cs
@@ -68,7 +68,6 @@ public class NiagaraEditor : ModuleRules
 			new string[] {
 				"MessageLog",
 				"Messaging",
-                "DerivedDataCache",
                 "LevelEditor",
 				"WorkspaceMenuStructure"
 			}
diff --git a/Engine/Plugins/FastBuildController/Source/FastbuildController.Build.cs b/Engine/Plugins/FastBuildController/Source/FastbuildController.Build.cs
--- a/Engine/Plugins/FastBuildController/Source/FastbuildController.Build.cs
+++ b/Engine/Plugins/FastBuildController/Source/FastbuildController.Build.cs
@@ -16,7 +16,7 @@ public class FastBuildController : ModuleRules
 		PrivateIncludePathModuleNames.AddRange(
 			new string[] {
 				"DistributedBuildInterface",
-				"Engine",
+				"TargetPlatform"
 			}
 		);
 	}
diff --git a/Engine/Plugins/Media/AndroidMedia/Source/AndroidMediaFactory/AndroidMediaFactory.Build.cs b/Engine/Plugins/Media/AndroidMedia/Source/AndroidMediaFactory/AndroidMediaFactory.Build.cs
--- a/Engine/Plugins/Media/AndroidMedia/Source/AndroidMediaFactory/AndroidMediaFactory.Build.cs
+++ b/Engine/Plugins/Media/AndroidMedia/Source/AndroidMediaFactory/AndroidMediaFactory.Build.cs
@@ -18,7 +18,6 @@ namespace UnrealBuildTool.Rules
 
 			PrivateIncludePathModuleNames.AddRange(
 				new string[] {
-					"AndroidMedia",
 					"Media",
 				});
 
diff --git a/Engine/Plugins/Media/AvfMedia/Source/AvfMediaFactory/AvfMediaFactory.Build.cs b/Engine/Plugins/Media/AvfMedia/Source/AvfMediaFactory/AvfMediaFactory.Build.cs
--- a/Engine/Plugins/Media/AvfMedia/Source/AvfMediaFactory/AvfMediaFactory.Build.cs
+++ b/Engine/Plugins/Media/AvfMedia/Source/AvfMediaFactory/AvfMediaFactory.Build.cs
@@ -18,7 +18,6 @@ namespace UnrealBuildTool.Rules
 
 			PrivateIncludePathModuleNames.AddRange(
 				new string[] {
-					"AvfMedia",
 					"Media",
 				});
 
diff --git a/Engine/Plugins/Media/ImgMedia/Source/ImgMediaEngine/ImgMediaEngine.Build.cs b/Engine/Plugins/Media/ImgMedia/Source/ImgMediaEngine/ImgMediaEngine.Build.cs
--- a/Engine/Plugins/Media/ImgMedia/Source/ImgMediaEngine/ImgMediaEngine.Build.cs
+++ b/Engine/Plugins/Media/ImgMedia/Source/ImgMediaEngine/ImgMediaEngine.Build.cs
@@ -14,8 +14,6 @@ namespace UnrealBuildTool.Rules
 
 			PrivateIncludePathModuleNames.AddRange(
 				new string[] {
-					"Media",
-					"ImgMedia",
 				});
 
 			// Are we using the engine?
diff --git a/Engine/Plugins/Media/ImgMedia/Source/ImgMediaFactory/ImgMediaFactory.Build.cs b/Engine/Plugins/Media/ImgMedia/Source/ImgMediaFactory/ImgMediaFactory.Build.cs
--- a/Engine/Plugins/Media/ImgMedia/Source/ImgMediaFactory/ImgMediaFactory.Build.cs
+++ b/Engine/Plugins/Media/ImgMedia/Source/ImgMediaFactory/ImgMediaFactory.Build.cs
@@ -19,7 +19,7 @@ namespace UnrealBuildTool.Rules
 			PrivateIncludePathModuleNames.AddRange(
 				new string[] {
 					"Media",
-					"ImgMedia",
+					"ImgMedia"
 				});
 
 			PublicDependencyModuleNames.AddRange(
diff --git a/Engine/Plugins/Media/MediaPlayerEditor/Source/MediaPlayerEditor/MediaPlayerEditor.Build.cs b/Engine/Plugins/Media/MediaPlayerEditor/Source/MediaPlayerEditor/MediaPlayerEditor.Build.cs
--- a/Engine/Plugins/Media/MediaPlayerEditor/Source/MediaPlayerEditor/MediaPlayerEditor.Build.cs
+++ b/Engine/Plugins/Media/MediaPlayerEditor/Source/MediaPlayerEditor/MediaPlayerEditor.Build.cs
@@ -50,7 +50,6 @@ public class MediaPlayerEditor : ModuleRules
 			new string[] {
 				"AssetTools",
 				"Media",
-				"WorkspaceMenuStructure",
 			});
 
 
diff --git a/Engine/Plugins/Media/MfMedia/Source/MfMediaFactory/MfMediaFactory.Build.cs b/Engine/Plugins/Media/MfMedia/Source/MfMediaFactory/MfMediaFactory.Build.cs
--- a/Engine/Plugins/Media/MfMedia/Source/MfMediaFactory/MfMediaFactory.Build.cs
+++ b/Engine/Plugins/Media/MfMedia/Source/MfMediaFactory/MfMediaFactory.Build.cs
@@ -21,7 +21,6 @@ namespace UnrealBuildTool.Rules
 			PrivateIncludePathModuleNames.AddRange(
 				new string[] {
 					"Media",
-					"MfMedia",
 				});
 
 			if (DoAllowHTTPSPlayback())
diff --git a/Engine/Plugins/Media/WmfMedia/Source/WmfMediaFactory/WmfMediaFactory.Build.cs b/Engine/Plugins/Media/WmfMedia/Source/WmfMediaFactory/WmfMediaFactory.Build.cs
--- a/Engine/Plugins/Media/WmfMedia/Source/WmfMediaFactory/WmfMediaFactory.Build.cs
+++ b/Engine/Plugins/Media/WmfMedia/Source/WmfMediaFactory/WmfMediaFactory.Build.cs
@@ -19,7 +19,6 @@ namespace UnrealBuildTool.Rules
 			PrivateIncludePathModuleNames.AddRange(
 				new string[] {
 					"Media",
-					"WmfMedia",
 				});
 
 			PublicDependencyModuleNames.AddRange(
diff --git a/Engine/Plugins/Online/OnlineFramework/Source/LoginFlow/LoginFlow.Build.cs b/Engine/Plugins/Online/OnlineFramework/Source/LoginFlow/LoginFlow.Build.cs
--- a/Engine/Plugins/Online/OnlineFramework/Source/LoginFlow/LoginFlow.Build.cs
+++ b/Engine/Plugins/Online/OnlineFramework/Source/LoginFlow/LoginFlow.Build.cs
@@ -37,8 +37,6 @@ public class LoginFlow : ModuleRules
 
 		PrivateIncludePathModuleNames.AddRange(
 			new string[] {
-				"Analytics",
-				"AnalyticsET",
 			}
 		);
 	}
diff --git a/Engine/Plugins/Runtime/AudioSynesthesia/Source/AudioSynesthesiaCore/AudioSynesthesiaCore.Build.cs b/Engine/Plugins/Runtime/AudioSynesthesia/Source/AudioSynesthesiaCore/AudioSynesthesiaCore.Build.cs
--- a/Engine/Plugins/Runtime/AudioSynesthesia/Source/AudioSynesthesiaCore/AudioSynesthesiaCore.Build.cs
+++ b/Engine/Plugins/Runtime/AudioSynesthesia/Source/AudioSynesthesiaCore/AudioSynesthesiaCore.Build.cs
@@ -24,7 +24,7 @@ namespace UnrealBuildTool.Rules
 
 			PrivateIncludePathModuleNames.AddRange(
 				new string[] {
-					"AudioMixer"
+					"AudioMixerCore"
 				}
 			);
 		}
diff --git a/Engine/Plugins/Runtime/RenderTrace/Source/RenderTrace.Build.cs b/Engine/Plugins/Runtime/RenderTrace/Source/RenderTrace.Build.cs
--- a/Engine/Plugins/Runtime/RenderTrace/Source/RenderTrace.Build.cs
+++ b/Engine/Plugins/Runtime/RenderTrace/Source/RenderTrace.Build.cs
@@ -9,7 +9,6 @@ public class RenderTrace : ModuleRules
 	{
 		PrivateIncludePathModuleNames.AddRange(
 			new string[] {
-				"TargetPlatform",
 			}
 		);
 
diff --git a/Engine/Source/Developer/AssetTools/AssetTools.Build.cs b/Engine/Source/Developer/AssetTools/AssetTools.Build.cs
--- a/Engine/Source/Developer/AssetTools/AssetTools.Build.cs
+++ b/Engine/Source/Developer/AssetTools/AssetTools.Build.cs
@@ -18,7 +18,9 @@ public class AssetTools : ModuleRules
 
 		PublicIncludePathModuleNames.AddRange(
 			new string[] {
+				"AssetDefinition",
 				"Merge",
+				"UnrealEd",
 			}
 		);
 
@@ -54,28 +56,14 @@ public class AssetTools : ModuleRules
 				"Analytics",
 				"ContentBrowser",
 				"CollectionManager",
-                "CurveAssetEditor",
 				"DesktopPlatform",
 				"EditorWidgets",
-				"GameProjectGeneration",
-                "ActorPickerMode",
 				"MainFrame",
 				"MessageLog",
 				"PackagesDialog",
 				"Persona",
-				"FontEditor",
-                "AudioEditor",
-                "SkeletonEditor",
-                "SkeletalMeshEditor",
                 "AnimationEditor",
-                "AnimationBlueprintEditor",
-                "AnimationModifiers",
-			    "TextureEditor",
-				"DataTableEditor",
 				"Cascade",
-				"PhysicsAssetEditor",
-				"CurveTableEditor",
-				"StaticMeshEditor",
 				"VirtualTexturingEditor"
 			}
 		);
diff --git a/Engine/Source/Developer/AutomationWindow/AutomationWindow.Build.cs b/Engine/Source/Developer/AutomationWindow/AutomationWindow.Build.cs
--- a/Engine/Source/Developer/AutomationWindow/AutomationWindow.Build.cs
+++ b/Engine/Source/Developer/AutomationWindow/AutomationWindow.Build.cs
@@ -43,7 +43,6 @@ namespace UnrealBuildTool.Rules
 
             PrivateIncludePathModuleNames.AddRange(
 				new string[] {
-					"Messaging",
 					"SessionServices",
 				}
 			);
diff --git a/Engine/Source/Developer/CollisionAnalyzer/CollisionAnalyzer.Build.cs b/Engine/Source/Developer/CollisionAnalyzer/CollisionAnalyzer.Build.cs
--- a/Engine/Source/Developer/CollisionAnalyzer/CollisionAnalyzer.Build.cs
+++ b/Engine/Source/Developer/CollisionAnalyzer/CollisionAnalyzer.Build.cs
@@ -9,7 +9,6 @@ public class CollisionAnalyzer : ModuleRules
         PrivateIncludePathModuleNames.AddRange(
             new string[] {
 				"DesktopPlatform",
-				"MainFrame",
 			}
         );
 
diff --git a/Engine/Source/Developer/LauncherServices/LauncherServices.Build.cs b/Engine/Source/Developer/LauncherServices/LauncherServices.Build.cs
--- a/Engine/Source/Developer/LauncherServices/LauncherServices.Build.cs
+++ b/Engine/Source/Developer/LauncherServices/LauncherServices.Build.cs
@@ -30,7 +30,6 @@ namespace UnrealBuildTool.Rules
 
 			PrivateIncludePathModuleNames.AddRange(
 				new string[] {
-					"Messaging",
 					"DerivedDataCache",
 				});
 		}
diff --git a/Engine/Source/Developer/LogVisualizer/LogVisualizer.Build.cs b/Engine/Source/Developer/LogVisualizer/LogVisualizer.Build.cs
--- a/Engine/Source/Developer/LogVisualizer/LogVisualizer.Build.cs
+++ b/Engine/Source/Developer/LogVisualizer/LogVisualizer.Build.cs
@@ -10,8 +10,6 @@ public class LogVisualizer : ModuleRules
         PrivateIncludePathModuleNames.AddRange(
 			new string[] {
 				"DesktopPlatform",
-				"MainFrame",
-				"LevelEditor"
 			}
 		);
 
diff --git a/Engine/Source/Developer/ProfilerClient/ProfilerClient.Build.cs b/Engine/Source/Developer/ProfilerClient/ProfilerClient.Build.cs
--- a/Engine/Source/Developer/ProfilerClient/ProfilerClient.Build.cs
+++ b/Engine/Source/Developer/ProfilerClient/ProfilerClient.Build.cs
@@ -27,7 +27,6 @@ namespace UnrealBuildTool.Rules
 				new string[] {
 					"Messaging",
                     "MessagingCommon",
-                    "SessionServices",
 				}
 			);
 		}
diff --git a/Engine/Source/Developer/ProjectLauncher/ProjectLauncher.Build.cs b/Engine/Source/Developer/ProjectLauncher/ProjectLauncher.Build.cs
--- a/Engine/Source/Developer/ProjectLauncher/ProjectLauncher.Build.cs
+++ b/Engine/Source/Developer/ProjectLauncher/ProjectLauncher.Build.cs
@@ -35,7 +35,6 @@ public class ProjectLauncher : ModuleRules
 
 		PrivateIncludePathModuleNames.AddRange(
 			new string[] {
-				"SessionServices",
 				"TargetDeviceServices",
 			}
 		);
diff --git a/Engine/Source/Developer/ScreenShotComparison/ScreenShotComparison.Build.cs b/Engine/Source/Developer/ScreenShotComparison/ScreenShotComparison.Build.cs
--- a/Engine/Source/Developer/ScreenShotComparison/ScreenShotComparison.Build.cs
+++ b/Engine/Source/Developer/ScreenShotComparison/ScreenShotComparison.Build.cs
@@ -35,7 +35,6 @@ public class ScreenShotComparison : ModuleRules
 
 		PrivateIncludePathModuleNames.AddRange(
 			new string[] {
-				"SessionServices",
 			}
 		);
 	}
diff --git a/Engine/Source/Developer/TextureBuild/TextureBuild.Build.cs b/Engine/Source/Developer/TextureBuild/TextureBuild.Build.cs
--- a/Engine/Source/Developer/TextureBuild/TextureBuild.Build.cs
+++ b/Engine/Source/Developer/TextureBuild/TextureBuild.Build.cs
@@ -8,7 +8,6 @@ public class TextureBuild : ModuleRules
 	{
 		PrivateIncludePathModuleNames.AddRange(new string[]
 		{
-			"CoreUObject",
 			"Engine",
 		});
 
diff --git a/Engine/Source/Developer/TextureBuildUtilities/TextureBuildUtilities.Build.cs b/Engine/Source/Developer/TextureBuildUtilities/TextureBuildUtilities.Build.cs
--- a/Engine/Source/Developer/TextureBuildUtilities/TextureBuildUtilities.Build.cs
+++ b/Engine/Source/Developer/TextureBuildUtilities/TextureBuildUtilities.Build.cs
@@ -11,8 +11,6 @@ public class TextureBuildUtilities : ModuleRules
 		// Include only , no link :
 		PrivateIncludePathModuleNames.AddRange(new string[]
 		{
-			"CoreUObject",
-			"Engine",
 			"TextureCompressor"
 		});
 
diff --git a/Engine/Source/Developer/TextureFormatASTC/TextureFormatASTC.Build.cs b/Engine/Source/Developer/TextureFormatASTC/TextureFormatASTC.Build.cs
--- a/Engine/Source/Developer/TextureFormatASTC/TextureFormatASTC.Build.cs
+++ b/Engine/Source/Developer/TextureFormatASTC/TextureFormatASTC.Build.cs
@@ -9,9 +9,8 @@ public class TextureFormatASTC : ModuleRules
 		PrivateIncludePathModuleNames.AddRange(new string[]
 		{
 			"DerivedDataCache",
-			"Engine",
-			"TargetPlatform",
 			"TextureCompressor",
+			"TextureFormat",
 		});
 		
 		PrivateDependencyModuleNames.AddRange(new string[]
diff --git a/Engine/Source/Developer/TextureFormatDXT/TextureFormatDXT.Build.cs b/Engine/Source/Developer/TextureFormatDXT/TextureFormatDXT.Build.cs
--- a/Engine/Source/Developer/TextureFormatDXT/TextureFormatDXT.Build.cs
+++ b/Engine/Source/Developer/TextureFormatDXT/TextureFormatDXT.Build.cs
@@ -10,7 +10,6 @@ public class TextureFormatDXT : ModuleRules
 		{
 			"DerivedDataCache",
 			"Engine",
-			"TargetPlatform",
 			"TextureCompressor",
 		});
 
diff --git a/Engine/Source/Developer/TextureFormatETC2/TextureFormatETC2.Build.cs b/Engine/Source/Developer/TextureFormatETC2/TextureFormatETC2.Build.cs
--- a/Engine/Source/Developer/TextureFormatETC2/TextureFormatETC2.Build.cs
+++ b/Engine/Source/Developer/TextureFormatETC2/TextureFormatETC2.Build.cs
@@ -9,9 +9,8 @@ public class TextureFormatETC2 : ModuleRules
 		PrivateIncludePathModuleNames.AddRange(new string[]
 		{
 			"DerivedDataCache",
-			"Engine",
-			"TargetPlatform",
 			"TextureCompressor",
+			"TextureFormat",
 		});
 
 		PrivateDependencyModuleNames.AddRange(new string[]
diff --git a/Engine/Source/Developer/TextureFormatIntelISPCTexComp/TextureFormatIntelISPCTexComp.Build.cs b/Engine/Source/Developer/TextureFormatIntelISPCTexComp/TextureFormatIntelISPCTexComp.Build.cs
--- a/Engine/Source/Developer/TextureFormatIntelISPCTexComp/TextureFormatIntelISPCTexComp.Build.cs
+++ b/Engine/Source/Developer/TextureFormatIntelISPCTexComp/TextureFormatIntelISPCTexComp.Build.cs
@@ -9,9 +9,8 @@ public class TextureFormatIntelISPCTexComp : ModuleRules
 		PrivateIncludePathModuleNames.AddRange(new string[]
 		{
 			"DerivedDataCache",
-			"Engine",
-			"TargetPlatform",
 			"TextureCompressor",
+			"TextureFormat",
 		});
 
 		PrivateDependencyModuleNames.AddRange(new string[]
diff --git a/Engine/Source/Developer/TextureFormatUncompressed/TextureFormatUncompressed.Build.cs b/Engine/Source/Developer/TextureFormatUncompressed/TextureFormatUncompressed.Build.cs
--- a/Engine/Source/Developer/TextureFormatUncompressed/TextureFormatUncompressed.Build.cs
+++ b/Engine/Source/Developer/TextureFormatUncompressed/TextureFormatUncompressed.Build.cs
@@ -9,9 +9,8 @@ public class TextureFormatUncompressed : ModuleRules
 		PrivateIncludePathModuleNames.AddRange(new string[]
 		{
 			"DerivedDataCache",
-			"Engine",
-			"TargetPlatform",
 			"TextureCompressor",
+			"TextureFormat",
 		});
 
 		PrivateDependencyModuleNames.AddRange(new string[]
diff --git a/Engine/Source/Developer/Windows/WindowsPlatformEditor/WindowsPlatformEditor.Build.cs b/Engine/Source/Developer/Windows/WindowsPlatformEditor/WindowsPlatformEditor.Build.cs
--- a/Engine/Source/Developer/Windows/WindowsPlatformEditor/WindowsPlatformEditor.Build.cs
+++ b/Engine/Source/Developer/Windows/WindowsPlatformEditor/WindowsPlatformEditor.Build.cs
@@ -29,7 +29,6 @@ public class WindowsPlatformEditor : ModuleRules
 
 		PrivateIncludePathModuleNames.AddRange(
 			new string[] {
-				"GameProjectGeneration",
 				"Settings",
 			}
 		);
diff --git a/Engine/Source/Editor/AIGraph/AIGraph.Build.cs b/Engine/Source/Editor/AIGraph/AIGraph.Build.cs
--- a/Engine/Source/Editor/AIGraph/AIGraph.Build.cs
+++ b/Engine/Source/Editor/AIGraph/AIGraph.Build.cs
@@ -17,8 +17,6 @@ public class AIGraph : ModuleRules
         PrivateIncludePathModuleNames.AddRange(
             new string[] {
 				"AssetRegistry",
-				"AssetTools",
-				"ContentBrowser"
 			}
         );
 
diff --git a/Engine/Source/Editor/AnimGraph/AnimGraph.Build.cs b/Engine/Source/Editor/AnimGraph/AnimGraph.Build.cs
--- a/Engine/Source/Editor/AnimGraph/AnimGraph.Build.cs
+++ b/Engine/Source/Editor/AnimGraph/AnimGraph.Build.cs
@@ -49,8 +49,6 @@ public class AnimGraph : ModuleRules
 
         PrivateIncludePathModuleNames.AddRange(
             new string[] {
-	            "SkeletonEditor",
-                "AdvancedPreviewScene",
                 "AnimationBlueprintEditor",
             }
         );
diff --git a/Engine/Source/Editor/AnimationBlueprintEditor/AnimationBlueprintEditor.Build.cs b/Engine/Source/Editor/AnimationBlueprintEditor/AnimationBlueprintEditor.Build.cs
--- a/Engine/Source/Editor/AnimationBlueprintEditor/AnimationBlueprintEditor.Build.cs
+++ b/Engine/Source/Editor/AnimationBlueprintEditor/AnimationBlueprintEditor.Build.cs
@@ -8,13 +8,7 @@ public class AnimationBlueprintEditor : ModuleRules
 	{
 		PrivateIncludePathModuleNames.AddRange(
 			new string[] {
-				"AssetRegistry", 
-				"MainFrame",
-				"DesktopPlatform",
                 "SkeletonEditor",
-                "ContentBrowser",
-                "AssetTools",
-                "AnimationEditor",
 				"MessageLog"
 			}
 		);
diff --git a/Engine/Source/Editor/CSVtoSVG/CSVtoSVG.Build.cs b/Engine/Source/Editor/CSVtoSVG/CSVtoSVG.Build.cs
--- a/Engine/Source/Editor/CSVtoSVG/CSVtoSVG.Build.cs
+++ b/Engine/Source/Editor/CSVtoSVG/CSVtoSVG.Build.cs
@@ -14,9 +14,6 @@ public class CSVtoSVG : ModuleRules
 
 		PrivateIncludePathModuleNames.AddRange(
 			new string[] {
-				"AssetRegistry",
-                "AssetTools",
-				"EditorWidgets",
 				"WorkspaceMenuStructure",
 			}
 		);
diff --git a/Engine/Source/Editor/ContentBrowser/ContentBrowser.Build.cs b/Engine/Source/Editor/ContentBrowser/ContentBrowser.Build.cs
--- a/Engine/Source/Editor/ContentBrowser/ContentBrowser.Build.cs
+++ b/Engine/Source/Editor/ContentBrowser/ContentBrowser.Build.cs
@@ -10,9 +10,7 @@ public class ContentBrowser : ModuleRules
 
 		PrivateIncludePathModuleNames.AddRange(
 			new string[] {
-				"GameProjectGeneration",
                 "MainFrame",
-				"PackagesDialog",
 			}
 		);
 
diff --git a/Engine/Source/Editor/DerivedDataEditor/DerivedDataEditor.Build.cs b/Engine/Source/Editor/DerivedDataEditor/DerivedDataEditor.Build.cs
--- a/Engine/Source/Editor/DerivedDataEditor/DerivedDataEditor.Build.cs
+++ b/Engine/Source/Editor/DerivedDataEditor/DerivedDataEditor.Build.cs
@@ -35,7 +35,6 @@ public class DerivedDataEditor : ModuleRules
 
 		PrivateIncludePathModuleNames.AddRange(
 			new string[] {
-				"MainFrame",
 			});
 		UnsafeTypeCastWarningLevel = WarningLevel.Error;
 	}
diff --git a/Engine/Source/Editor/DeviceProfileServices/DeviceProfileServices.Build.cs b/Engine/Source/Editor/DeviceProfileServices/DeviceProfileServices.Build.cs
--- a/Engine/Source/Editor/DeviceProfileServices/DeviceProfileServices.Build.cs
+++ b/Engine/Source/Editor/DeviceProfileServices/DeviceProfileServices.Build.cs
@@ -27,7 +27,6 @@ public class DeviceProfileServices : ModuleRules
 
 		PrivateIncludePathModuleNames.AddRange(
 			new string[] {
-				"TargetDeviceServices",
 			});
 	}
 }
diff --git a/Engine/Source/Editor/EnvironmentLightingViewer/EnvironmentLightingViewer.Build.cs b/Engine/Source/Editor/EnvironmentLightingViewer/EnvironmentLightingViewer.Build.cs
--- a/Engine/Source/Editor/EnvironmentLightingViewer/EnvironmentLightingViewer.Build.cs
+++ b/Engine/Source/Editor/EnvironmentLightingViewer/EnvironmentLightingViewer.Build.cs
@@ -8,9 +8,6 @@ public class EnvironmentLightingViewer : ModuleRules
 	{
 		PrivateIncludePathModuleNames.AddRange(
 			new string[] {
-				"AssetRegistry",
-                "AssetTools",
-				"EditorWidgets",
 				"WorkspaceMenuStructure"
 			}
 		);
diff --git a/Engine/Source/Editor/FontEditor/FontEditor.Build.cs b/Engine/Source/Editor/FontEditor/FontEditor.Build.cs
--- a/Engine/Source/Editor/FontEditor/FontEditor.Build.cs
+++ b/Engine/Source/Editor/FontEditor/FontEditor.Build.cs
@@ -27,7 +27,6 @@ public class FontEditor : ModuleRules
 
 		PrivateIncludePathModuleNames.AddRange(
 			new string[] {
-				"MainFrame",
 				"EditorFramework",
 			}
 		);
diff --git a/Engine/Source/Editor/HierarchicalLODOutliner/HierarchicalLODOutliner.Build.cs b/Engine/Source/Editor/HierarchicalLODOutliner/HierarchicalLODOutliner.Build.cs
--- a/Engine/Source/Editor/HierarchicalLODOutliner/HierarchicalLODOutliner.Build.cs
+++ b/Engine/Source/Editor/HierarchicalLODOutliner/HierarchicalLODOutliner.Build.cs
@@ -24,7 +24,6 @@ public class HierarchicalLODOutliner : ModuleRules
 
 		PrivateIncludePathModuleNames.AddRange(
 			new string[] {
-				"LevelEditor"
 				}
 		);
 
diff --git a/Engine/Source/Editor/InternationalizationSettings/InternationalizationSettings.Build.cs b/Engine/Source/Editor/InternationalizationSettings/InternationalizationSettings.Build.cs
--- a/Engine/Source/Editor/InternationalizationSettings/InternationalizationSettings.Build.cs
+++ b/Engine/Source/Editor/InternationalizationSettings/InternationalizationSettings.Build.cs
@@ -23,8 +23,6 @@ namespace UnrealBuildTool.Rules
 
 			PrivateIncludePathModuleNames.AddRange(
 				new string[] {
-					"Settings",
-                    "SettingsEditor"
 				}
 			);
 
diff --git a/Engine/Source/Editor/Kismet/Kismet.Build.cs b/Engine/Source/Editor/Kismet/Kismet.Build.cs
--- a/Engine/Source/Editor/Kismet/Kismet.Build.cs
+++ b/Engine/Source/Editor/Kismet/Kismet.Build.cs
@@ -13,7 +13,6 @@ public class Kismet : ModuleRules
                 "BlueprintRuntime",
                 "ClassViewer",
 				"Analytics",
-                "DerivedDataCache",
                 "LevelEditor",
 				"GameProjectGeneration",
 				"SourceCodeAccess",
diff --git a/Engine/Source/Editor/LandscapeEditor/LandscapeEditor.Build.cs b/Engine/Source/Editor/LandscapeEditor/LandscapeEditor.Build.cs
--- a/Engine/Source/Editor/LandscapeEditor/LandscapeEditor.Build.cs
+++ b/Engine/Source/Editor/LandscapeEditor/LandscapeEditor.Build.cs
@@ -40,7 +40,6 @@ public class LandscapeEditor : ModuleRules
 			new string[] {
 				"MainFrame",
 				"DesktopPlatform",
-				"ContentBrowser",
                 "AssetTools",
 				"LevelEditor"
 			}
diff --git a/Engine/Source/Editor/LevelEditor/LevelEditor.Build.cs b/Engine/Source/Editor/LevelEditor/LevelEditor.Build.cs
--- a/Engine/Source/Editor/LevelEditor/LevelEditor.Build.cs
+++ b/Engine/Source/Editor/LevelEditor/LevelEditor.Build.cs
@@ -20,7 +20,6 @@ public class LevelEditor : ModuleRules
                 "PlacementMode",
 				"SlateReflector",
                 "PortalServices",
-                "Persona",
 				"MergeActors",
 				"Layers",
 				"WorldBrowser",
diff --git a/Engine/Source/Editor/MainFrame/MainFrame.Build.cs b/Engine/Source/Editor/MainFrame/MainFrame.Build.cs
--- a/Engine/Source/Editor/MainFrame/MainFrame.Build.cs
+++ b/Engine/Source/Editor/MainFrame/MainFrame.Build.cs
@@ -50,15 +50,12 @@ public class MainFrame : ModuleRules
 
 		PrivateIncludePathModuleNames.AddRange(
 			new string[] {
-				"AssetTools",
-				"DerivedDataCache",
 				"LauncherPlatform",
 				"GameProjectGeneration",
 				"ProjectTargetPlatformEditor",
 				"LevelEditor",
 				"Settings",
 				"SourceCodeAccess",
-				"LocalizationDashboard", // not required but causes circular depends issues on Linux/Mac
 				"HotReload",
 			}
 		);
diff --git a/Engine/Source/Editor/MergeActors/MergeActors.Build.cs b/Engine/Source/Editor/MergeActors/MergeActors.Build.cs
--- a/Engine/Source/Editor/MergeActors/MergeActors.Build.cs
+++ b/Engine/Source/Editor/MergeActors/MergeActors.Build.cs
@@ -13,7 +13,6 @@ public class MergeActors : ModuleRules
                 "Documentation",
                 "MeshUtilities",
                 "PropertyEditor",
-                "RawMesh",
                 "WorkspaceMenuStructure",
                 "MeshReductionInterface",
                 "MeshMergeUtilities",
diff --git a/Engine/Source/Editor/MovieSceneCaptureDialog/MovieSceneCaptureDialog.Build.cs b/Engine/Source/Editor/MovieSceneCaptureDialog/MovieSceneCaptureDialog.Build.cs
--- a/Engine/Source/Editor/MovieSceneCaptureDialog/MovieSceneCaptureDialog.Build.cs
+++ b/Engine/Source/Editor/MovieSceneCaptureDialog/MovieSceneCaptureDialog.Build.cs
@@ -28,7 +28,6 @@ public class MovieSceneCaptureDialog : ModuleRules
 
 		PrivateIncludePathModuleNames.AddRange(
 			new string[] {
-				"LevelEditor"
 			}
 		);
 	}
diff --git a/Engine/Source/Editor/PListEditor/PListEditor.Build.cs b/Engine/Source/Editor/PListEditor/PListEditor.Build.cs
--- a/Engine/Source/Editor/PListEditor/PListEditor.Build.cs
+++ b/Engine/Source/Editor/PListEditor/PListEditor.Build.cs
@@ -23,7 +23,6 @@ public class PListEditor : ModuleRules
 
 		PrivateIncludePathModuleNames.AddRange(
 			new string[] {
-				"MainFrame",
 				"WorkspaceMenuStructure",
 			}
 		);
diff --git a/Engine/Source/Editor/PluginWarden/PluginWarden.Build.cs b/Engine/Source/Editor/PluginWarden/PluginWarden.Build.cs
--- a/Engine/Source/Editor/PluginWarden/PluginWarden.Build.cs
+++ b/Engine/Source/Editor/PluginWarden/PluginWarden.Build.cs
@@ -8,7 +8,6 @@ public class PluginWarden : ModuleRules
 	{
 		PrivateIncludePathModuleNames.AddRange(
 			new string[] {
-				"LauncherServices",
 			}
 		);
 
diff --git a/Engine/Source/Editor/RenderResourceViewer/RenderResourceViewer.Build.cs b/Engine/Source/Editor/RenderResourceViewer/RenderResourceViewer.Build.cs
--- a/Engine/Source/Editor/RenderResourceViewer/RenderResourceViewer.Build.cs
+++ b/Engine/Source/Editor/RenderResourceViewer/RenderResourceViewer.Build.cs
@@ -11,8 +11,6 @@ public class RenderResourceViewer : ModuleRules
 			{
 				"WorkspaceMenuStructure",
 				"ContentBrowser",
-				"UnrealEd",
-				"AssetDefinition"
 			}
 		);
 
diff --git a/Engine/Source/Editor/Sequencer/Sequencer.Build.cs b/Engine/Source/Editor/Sequencer/Sequencer.Build.cs
--- a/Engine/Source/Editor/Sequencer/Sequencer.Build.cs
+++ b/Engine/Source/Editor/Sequencer/Sequencer.Build.cs
@@ -63,9 +63,7 @@ public class Sequencer : ModuleRules
 
 		PrivateIncludePathModuleNames.AddRange(
 			new string[] {
-				"Kismet",
 				"LevelEditor",
-				"MainFrame",
 				"DesktopPlatform",
 				}
 			);
diff --git a/Engine/Source/Editor/StaticMeshEditor/StaticMeshEditor.Build.cs b/Engine/Source/Editor/StaticMeshEditor/StaticMeshEditor.Build.cs
--- a/Engine/Source/Editor/StaticMeshEditor/StaticMeshEditor.Build.cs
+++ b/Engine/Source/Editor/StaticMeshEditor/StaticMeshEditor.Build.cs
@@ -16,8 +16,6 @@ public class StaticMeshEditor : ModuleRules
 		PrivateIncludePathModuleNames.AddRange(
 			new string[] {
 				"AssetRegistry",
-				"Kismet",
-				"EditorWidgets",
                 "MeshReductionInterface",
             }
 		);
diff --git a/Engine/Source/Editor/StructViewer/StructViewer.Build.cs b/Engine/Source/Editor/StructViewer/StructViewer.Build.cs
--- a/Engine/Source/Editor/StructViewer/StructViewer.Build.cs
+++ b/Engine/Source/Editor/StructViewer/StructViewer.Build.cs
@@ -9,7 +9,6 @@ public class StructViewer : ModuleRules
 		PrivateIncludePathModuleNames.AddRange(
 			new string[] {
 				"AssetRegistry",
-                "AssetTools",
 				"EditorWidgets",
 				"WorkspaceMenuStructure",
 			}
diff --git a/Engine/Source/Editor/TurnkeySupport/TurnkeySupport.Build.cs b/Engine/Source/Editor/TurnkeySupport/TurnkeySupport.Build.cs
--- a/Engine/Source/Editor/TurnkeySupport/TurnkeySupport.Build.cs
+++ b/Engine/Source/Editor/TurnkeySupport/TurnkeySupport.Build.cs
@@ -49,13 +49,10 @@ public class TurnkeySupport : ModuleRules
 			);
 			PrivateIncludePathModuleNames.AddRange(
 				new string[] {
-					"AssetTools",
 					"GameProjectGeneration",
 					"ProjectTargetPlatformEditor",
 					"LevelEditor",
 					"Settings",
-					"SourceCodeAccess",
-					"LocalizationDashboard",
 	 				"MainFrame",
 				}
 			);
diff --git a/Engine/Source/Editor/WorldBrowser/WorldBrowser.Build.cs b/Engine/Source/Editor/WorldBrowser/WorldBrowser.Build.cs
--- a/Engine/Source/Editor/WorldBrowser/WorldBrowser.Build.cs
+++ b/Engine/Source/Editor/WorldBrowser/WorldBrowser.Build.cs
@@ -10,7 +10,6 @@ public class WorldBrowser : ModuleRules
             new string[] {
                 "AssetRegistry",
 				"AssetTools",
-                "ContentBrowser",
                 "MeshUtilities",
                 "MeshMergeUtilities",
             }
diff --git a/Engine/Source/Runtime/Portal/Proxies/PortalProxies.Build.cs b/Engine/Source/Runtime/Portal/Proxies/PortalProxies.Build.cs
--- a/Engine/Source/Runtime/Portal/Proxies/PortalProxies.Build.cs
+++ b/Engine/Source/Runtime/Portal/Proxies/PortalProxies.Build.cs
@@ -30,7 +30,6 @@ namespace UnrealBuildTool.Rules
 
             PrivateIncludePathModuleNames.AddRange(
                 new string[] {
-                    "Messaging",
                     "MessagingRpc",
                     "PortalServices",
                 }
diff --git a/Engine/Source/Runtime/Portal/Rpc/PortalRpc.Build.cs b/Engine/Source/Runtime/Portal/Rpc/PortalRpc.Build.cs
--- a/Engine/Source/Runtime/Portal/Rpc/PortalRpc.Build.cs
+++ b/Engine/Source/Runtime/Portal/Rpc/PortalRpc.Build.cs
@@ -24,7 +24,6 @@ namespace UnrealBuildTool.Rules
                 new string[] {
                     "Messaging",
                     "MessagingCommon",
-                    "PortalServices",
                 }
             );
 		}
