#UE Update internal buildmachine root path in cooking filepath length calculation

diff --git a/Engine/Source/Developer/AssetTools/Private/AssetViewUtils.cpp b/Engine/Source/Developer/AssetTools/Private/AssetViewUtils.cpp
--- a/Engine/Source/Developer/AssetTools/Private/AssetViewUtils.cpp
+++ b/Engine/Source/Developer/AssetTools/Private/AssetViewUtils.cpp
@@ -1514,7 +1514,7 @@ int32 AssetViewUtils::GetPackageLengthForCooking(const FString& PackageName, boo
 			int32 CookPathRelativeToTargetRootLen = CookSubPath.Len();
 			CookPathRelativeToTargetRootLen += (AbsoluteTargetPath.Len() - AbsoluteUERootPath.Len());
 
-			int32 InternalCookPathLen = FCString::Strlen(TEXT("D:/BuildFarm/buildmachine_++depot+UE-Releases+XX.XX/")) + CookPathRelativeToTargetRootLen + AssetPathRelativeToCookRootLen;
+			int32 InternalCookPathLen = FCString::Strlen(TEXT("d:/build/XXX+RXX.XX+Inc/Sync")) + CookPathRelativeToTargetRootLen + AssetPathRelativeToCookRootLen;
 
 			// Only add game name padding to project plugins so that they can ported to other projects
 			if (PluginContainingAsset && bIsProjectAsset)
