#UE Internal, in-UE-root, gamefeatureplugins no longer pad their GameName to 20 characters to prevent filepath length limits from being exceeded. This rule was added to allow plugins to be used in other projects that may have a longer name, but GameFeaturePlugins cannot be used in other projects, since they depend in /Game/ and project-level modules

diff --git a/Engine/Source/Developer/AssetTools/Private/AssetViewUtils.cpp b/Engine/Source/Developer/AssetTools/Private/AssetViewUtils.cpp
--- a/Engine/Source/Developer/AssetTools/Private/AssetViewUtils.cpp
+++ b/Engine/Source/Developer/AssetTools/Private/AssetViewUtils.cpp
@@ -1516,8 +1516,10 @@ int32 AssetViewUtils::GetPackageLengthForCooking(const FString& PackageName, boo
 
 			int32 InternalCookPathLen = FCString::Strlen(TEXT("d:/build/XXX+RXX.XX+Inc/Sync")) + CookPathRelativeToTargetRootLen + AssetPathRelativeToCookRootLen;
 
-			// Only add game name padding to project plugins so that they can ported to other projects
-			if (PluginContainingAsset && bIsProjectAsset)
+			// Only add game name padding to non-GameFeature project plugins so that they can ported to other projects.
+			// GameFeaturePlugins are not designed to be ported between projects since they can depend on /Game/ assets and non-plugin code.
+			const bool bIsGameFeaturePlugin = AbsolutePathToAsset.Contains(TEXT("/Plugins/GameFeatures/"));
+			if (PluginContainingAsset && bIsProjectAsset && !bIsGameFeaturePlugin)
 			{
 				// We assume the game name is 20 characters (the maximum allowed) to make sure that content can be ported between projects
 				constexpr int32 MaxGameNameLen = 20;
