Fixed asset rename not updating the soft object paths for packages that were just renamed and left a redirector.

The issue was that when moving folder A and B, where A has assets that depends on assets in B, we'd ask the asset registry for the B assets referencers and it was returning to original packages which are now a redirectors. We'd then try to patch soft object paths on a redirector instead of on an actual asset.

We now force a rescan of the moved folders so that the asset registry is up-to-date throughout the move folders process.

diff --git a/Engine/Source/Developer/AssetTools/Private/AssetViewUtils.cpp b/Engine/Source/Developer/AssetTools/Private/AssetViewUtils.cpp
--- a/Engine/Source/Developer/AssetTools/Private/AssetViewUtils.cpp
+++ b/Engine/Source/Developer/AssetTools/Private/AssetViewUtils.cpp
@@ -711,11 +711,21 @@ bool AssetViewUtils::MoveFolders(const TArray<FString>& InSourcePathNames, const
 			MoveAssets( PathIt.Value(), Destination, PathIt.Key() );
 		}
 
+		TArray<FString> PathsToRescan;
+		PathsToRescan.Add(Destination);
+
 		// Attempt to remove the old path
 		if (DeleteEmptyFolderFromDisk(SourcePath))
 		{
 			AssetRegistryModule.Get().RemovePath(SourcePath);
 		}
+		else if (PathIt.Value().Num() > 0)
+		{
+			PathsToRescan.Add(SourcePath);
+		}
+
+		const bool bForceRescan = true;
+		AssetRegistryModule.Get().ScanPathsSynchronous(PathsToRescan, bForceRescan);
 
 		const TOptional<FLinearColor> FolderColor = GetPathColor(SourcePath);
 		if (FolderColor.IsSet())
