[Content Browser] Fix Clear Thumbnail not working if asset not already loaded into memory.

diff --git a/Engine/Source/Developer/AssetTools/Private/AssetViewUtils.cpp b/Engine/Source/Developer/AssetTools/Private/AssetViewUtils.cpp
--- a/Engine/Source/Developer/AssetTools/Private/AssetViewUtils.cpp
+++ b/Engine/Source/Developer/AssetTools/Private/AssetViewUtils.cpp
@@ -923,9 +923,9 @@ void AssetViewUtils::ClearCustomThumbnails(const TArray<FAssetData>& InAssetsToA
 		{
 			//assign the thumbnail and dirty
 			const FString ObjectFullName = CurrentAsset.GetFullName();
-			const FString PackageName    = CurrentAsset.PackageName.ToString();
 
-			UPackage* AssetPackage = FindObject<UPackage>( NULL, *PackageName );
+			// Load the asset so that we can find and clear the thumbnail
+			UPackage* AssetPackage = CurrentAsset.GetPackage();
 			if ( ensure(AssetPackage) )
 			{
 				ThumbnailTools::CacheEmptyThumbnail( ObjectFullName, AssetPackage);
