Added a missing meta data to tell the Blueprints that the AssetTools::MigratePackages function shouldn't be pure.

diff --git a/Engine/Source/Developer/AssetTools/Public/IAssetTools.h b/Engine/Source/Developer/AssetTools/Public/IAssetTools.h
--- a/Engine/Source/Developer/AssetTools/Public/IAssetTools.h
+++ b/Engine/Source/Developer/AssetTools/Public/IAssetTools.h
@@ -516,7 +516,7 @@ public:
 	virtual void MigratePackages(const TArray<FName>& PackageNamesToMigrate) const = 0;
 
 	/* Migrate packages and dependencies to another folder */
-	UFUNCTION(BlueprintCallable, Category = "Editor Scripting | Asset Tools")
+	UFUNCTION(BlueprintCallable, Category = "Editor Scripting | Asset Tools", BlueprintPure = false)
 	virtual void MigratePackages(const TArray<FName>& PackageNamesToMigrate, const FString& DestinationPath, const struct FMigrationOptions& Options = FMigrationOptions()) const = 0;
 
 	/**
