Truncation warnings fixes for AssetTools

diff --git a/Engine/Source/Developer/AssetTools/AssetTools.Build.cs b/Engine/Source/Developer/AssetTools/AssetTools.Build.cs
--- a/Engine/Source/Developer/AssetTools/AssetTools.Build.cs
+++ b/Engine/Source/Developer/AssetTools/AssetTools.Build.cs
@@ -98,5 +98,7 @@ public class AssetTools : ModuleRules
 				"UnrealEd",
 			}
 		);
+
+		UnsafeTypeCastWarningLevel = WarningLevel.Error;
 	}
 }
diff --git a/Engine/Source/Developer/AssetTools/Private/AssetFixUpRedirectors.cpp b/Engine/Source/Developer/AssetTools/Private/AssetFixUpRedirectors.cpp
--- a/Engine/Source/Developer/AssetTools/Private/AssetFixUpRedirectors.cpp
+++ b/Engine/Source/Developer/AssetTools/Private/AssetFixUpRedirectors.cpp
@@ -208,7 +208,7 @@ bool FAssetFixUpRedirectors::UpdatePackageStatus(const TArray<FRedirectorRefs>&
 
 void FAssetFixUpRedirectors::LoadReferencingPackages(TArray<FRedirectorRefs>& RedirectorsToFix, TArray<UPackage*>& OutReferencingPackagesToSave, TArray<UPackage*>& OutLoadedPackages) const
 {
-	FScopedSlowTask SlowTask( RedirectorsToFix.Num(), LOCTEXT( "LoadingReferencingPackages", "Loading Referencing Packages..." ) );
+	FScopedSlowTask SlowTask( static_cast<float>(RedirectorsToFix.Num()), LOCTEXT( "LoadingReferencingPackages", "Loading Referencing Packages..." ) );
 	SlowTask.MakeDialog();
 
 	ISourceControlProvider& SourceControlProvider = ISourceControlModule::Get().GetProvider();
diff --git a/Engine/Source/Developer/AssetTools/Private/AssetTools.cpp b/Engine/Source/Developer/AssetTools/Private/AssetTools.cpp
--- a/Engine/Source/Developer/AssetTools/Private/AssetTools.cpp
+++ b/Engine/Source/Developer/AssetTools/Private/AssetTools.cpp
@@ -965,13 +965,13 @@ namespace UE::AssetTools::Private
 
 		static void LoadInstancedPackages(FPackageMigrationContext& PackageMigrationContext, FPackageMigrationImplContext& MigrationImplContext)
 		{
-			const int32 NumOfPackageToLoad = Algo::CountIf(PackageMigrationContext.MigrationPackagesData, [](const FPackageMigrationContext::FMigrationPackageData& MigrationPackageData)
+			const int32 NumOfPackageToLoad = static_cast<int32>(Algo::CountIf(PackageMigrationContext.MigrationPackagesData, [](const FPackageMigrationContext::FMigrationPackageData& MigrationPackageData)
 				{
 					return MigrationPackageData.bNeedInstancedLoad || MigrationPackageData.bNeedToBeSaveMigrated;
-				});
+				}));
 
 			// We give 40% of the progress weight to the loading
-			const float ProgressPerItemToload = 40.f / NumOfPackageToLoad;
+			const float ProgressPerItemToload = 40.f / static_cast<float>(NumOfPackageToLoad);
 
 			uint32 ProgressCount = 0;
 
@@ -1031,13 +1031,13 @@ namespace UE::AssetTools::Private
 
 			uint32 ProgressCount = 0;
 
-			const int32 NumOfPackageToSave = Algo::CountIf(PackageMigrationContext.MigrationPackagesData, [](const FPackageMigrationContext::FMigrationPackageData& MigrationPackageData)
+			const int32 NumOfPackageToSave = static_cast<int32>(Algo::CountIf(PackageMigrationContext.MigrationPackagesData, [](const FPackageMigrationContext::FMigrationPackageData& MigrationPackageData)
 				{
 					return MigrationPackageData.bNeedToBeSaveMigrated;
-				});
+				}));
 
 			// We give 60% of the progress weight to the saving of the assets
-			const float ProgressPerItemToSave = 60.f / NumOfPackageToSave;
+			const float ProgressPerItemToSave = 60.f / static_cast<float>(NumOfPackageToSave);
 
 			for (const FPackageMigrationContext::FMigrationPackageData& MigrationPackageData : PackageMigrationContext.MigrationPackagesData)
 			{
@@ -2159,7 +2159,7 @@ bool UAssetToolsImpl::AdvancedCopyPackages(
 		ExistingObjectSet.Reserve(SourceAndDestPackages.Num());
 		NewObjectSet.Reserve(SourceAndDestPackages.Num());
 
-		TUniquePtr<FScopedSlowTask> LoopProgress = MakeUnique<FScopedSlowTask>(SourceAndDestPackages.Num(), LOCTEXT("AdvancedCopyPackages.CopyingFilesAndDependencies", "Copying Files and Dependencies..."));
+		TUniquePtr<FScopedSlowTask> LoopProgress = MakeUnique<FScopedSlowTask>(static_cast<float>(SourceAndDestPackages.Num()), LOCTEXT("AdvancedCopyPackages.CopyingFilesAndDependencies", "Copying Files and Dependencies..."));
 		LoopProgress->MakeDialog();
 
 		for (const auto& Package : SourceAndDestPackages)
@@ -2223,7 +2223,7 @@ bool UAssetToolsImpl::AdvancedCopyPackages(
 		ObjectTools::GatherSubObjectsForReferenceReplacement(NewObjectSet, ExistingObjectSet, ObjectsAndSubObjectsToReplaceWithin);
 
 		LoopProgress.Reset(); // Note: Reset first as FScopedSlowTask asserts about out-of-order scoping if assigning over an existing FScopedSlowTask
-		LoopProgress = MakeUnique<FScopedSlowTask>(SuccessfullyCopiedSourcePackages.Num(), LOCTEXT("AdvancedCopyPackages.ReplacingAssetReferences", "Replacing Asset References..."));
+		LoopProgress = MakeUnique<FScopedSlowTask>(static_cast<float>(SuccessfullyCopiedSourcePackages.Num()), LOCTEXT("AdvancedCopyPackages.ReplacingAssetReferences", "Replacing Asset References..."));
 		LoopProgress->MakeDialog();
 
 		TMap<UObject*, TArray<UObject*, TInlineAllocator<1>>> Consolidations;
@@ -2482,7 +2482,7 @@ TArray<UObject*> UAssetToolsImpl::ImportAssetsAutomated(const UAutomatedAssetImp
 
 void UAssetToolsImpl::ImportAssetTasks(const TArray<UAssetImportTask*>& ImportTasks)
 {
-	FScopedSlowTask SlowTask(ImportTasks.Num(), LOCTEXT("ImportSlowTask", "Importing"));
+	FScopedSlowTask SlowTask(static_cast<float>(ImportTasks.Num()), LOCTEXT("ImportSlowTask", "Importing"));
 	SlowTask.MakeDialog();
 
 	FAssetImportParams Params;
@@ -3070,7 +3070,7 @@ TArray<UObject*> UAssetToolsImpl::ImportAssetsInternal(const TArray<FString>& Fi
 		return UInterchangeManager::IsInterchangeImportEnabled() && (SpecifiedFactory == nullptr);
 	}();
 
-	FScopedSlowTask SlowTask(ValidFiles.Num(), LOCTEXT("ImportSlowTask", "Importing"), !bUseInterchangeFramework);
+	FScopedSlowTask SlowTask(static_cast<float>(ValidFiles.Num()), LOCTEXT("ImportSlowTask", "Importing"), !bUseInterchangeFramework);
 
 	if (!bUseInterchangeFramework && ValidFiles.Num() > 1)
 	{	
@@ -3243,7 +3243,7 @@ TArray<UObject*> UAssetToolsImpl::ImportAssetsInternal(const TArray<FString>& Fi
 		else
 		{
 			//Complete the slow task
-			SlowTask.CompletedWork = FilesAndDestinations.Num();
+			SlowTask.CompletedWork = static_cast<float>(FilesAndDestinations.Num());
 		}
 	}
 
@@ -4386,7 +4386,7 @@ void UAssetToolsImpl::PerformMigratePackages(TArray<FName> PackageNamesToMigrate
 	TSet<FString> ExternalObjectsPaths;
 	TSet<FName> ExcludedDependencies;
 	{
-		FScopedSlowTask SlowTask( PackageNamesToMigrate.Num(), LOCTEXT( "MigratePackages_GatheringDependencies", "Gathering Dependencies..." ) );
+		FScopedSlowTask SlowTask(static_cast<float>(PackageNamesToMigrate.Num()), LOCTEXT( "MigratePackages_GatheringDependencies", "Gathering Dependencies..." ) );
 		SlowTask.MakeDialog();
 
 		TFunction<bool (FName)> ShouldExcludePackage = [this](FName PackageName) { return !CanMigratePackage(PackageName); };
@@ -4964,7 +4964,7 @@ void UAssetToolsImpl::MigratePackages_ReportConfirmed(TSharedPtr<TArray<ReportPa
 			// Resolve the packages to migrate to assets
 			TArray<UObject*> SrcObjects;
 			{
-				FScopedSlowTask SlowTask(AllPackageNamesToMove.Num(), LOCTEXT("MigratePackages_Loading", "Loading Packages..."));
+				FScopedSlowTask SlowTask(static_cast<float>(AllPackageNamesToMove.Num()), LOCTEXT("MigratePackages_Loading", "Loading Packages..."));
 				SlowTask.MakeDialog();
 				for (const FName& SrcPackage : AllPackageNamesToMove)
 				{
@@ -5007,7 +5007,7 @@ void UAssetToolsImpl::MigratePackages_ReportConfirmed(TSharedPtr<TArray<ReportPa
 			TArray<UObject*> TempObjects;
 			TMap<UObject*, UObject*> ReplacementMap;
 			{
-				FScopedSlowTask SlowTask(SrcObjects.Num(), LOCTEXT("MigratePackages_Duplicate", "Duplicating Assets..."));
+				FScopedSlowTask SlowTask(static_cast<float>(SrcObjects.Num()), LOCTEXT("MigratePackages_Duplicate", "Duplicating Assets..."));
 				SlowTask.MakeDialog();
 
 				TSet<UPackage*> PackagesUserRefusedToFullyLoad;
@@ -5047,7 +5047,7 @@ void UAssetToolsImpl::MigratePackages_ReportConfirmed(TSharedPtr<TArray<ReportPa
 
 			// Update references between TempObjects (to reference each other)
 			{
-				FScopedSlowTask SlowTask(TempObjects.Num(), LOCTEXT("MigratePackages_ReplaceReferences", "Replacing References..."));
+				FScopedSlowTask SlowTask(static_cast<float>(TempObjects.Num()), LOCTEXT("MigratePackages_ReplaceReferences", "Replacing References..."));
 				SlowTask.MakeDialog();
 
 				for (int i=0; i<TempObjects.Num(); ++i)
@@ -5061,7 +5061,7 @@ void UAssetToolsImpl::MigratePackages_ReportConfirmed(TSharedPtr<TArray<ReportPa
 
 			// Save fixed up packages to the migrated folder, and update the set of files to copy to be those migrated packages
 			{
-				FScopedSlowTask SlowTask(TempObjects.Num(), LOCTEXT("MigratePackages_ReplaceReferences", "Replacing References..."));
+				FScopedSlowTask SlowTask(static_cast<float>(TempObjects.Num()), LOCTEXT("MigratePackages_ReplaceReferences", "Replacing References..."));
 				SlowTask.MakeDialog();
 
 				TSet<FName> NewPackageNamesToMove;
@@ -5108,7 +5108,7 @@ void UAssetToolsImpl::MigratePackages_ReportConfirmed(TSharedPtr<TArray<ReportPa
 		FString CopyErrors;
 
 		{
-			FScopedSlowTask SlowTask(AllPackageNamesToMove.Num(), LOCTEXT("MigratePackages_CopyingFiles", "Copying Files..."));
+			FScopedSlowTask SlowTask(static_cast<float>(AllPackageNamesToMove.Num()), LOCTEXT("MigratePackages_CopyingFiles", "Copying Files..."));
 			SlowTask.MakeDialog();
 
 			for ( const FName& PackageNameToMove : AllPackageNamesToMove )
@@ -5211,7 +5211,7 @@ void UAssetToolsImpl::MigratePackages_ReportConfirmed(TSharedPtr<TArray<ReportPa
 		// If we are consolidating lost packages, we are copying temporary packages, so clean them up.
 		if (!LostPackages.IsEmpty())
 		{
-			FScopedSlowTask SlowTask(AllPackageNamesToMove.Num(), LOCTEXT("MigratePackages_CleaningUp", "Cleaning Up..."));
+			FScopedSlowTask SlowTask(static_cast<float>(AllPackageNamesToMove.Num()), LOCTEXT("MigratePackages_CleaningUp", "Cleaning Up..."));
 			SlowTask.MakeDialog();
 
 			TArray<UObject*> ObjectsToDelete;
@@ -5247,7 +5247,7 @@ void UAssetToolsImpl::MigratePackages_ReportConfirmed(TSharedPtr<TArray<ReportPa
 					ISourceControlProvider& SourceControlProvider = ISourceControlModule::Get().GetProvider();
 					if(SourceControlProvider.Execute(ISourceControlOperation::Create<FMarkForAdd>(), SuccessfullyCopiedFiles) == ECommandResult::Failed)
 					{
-						FScopedSlowTask SlowTask(SuccessfullyCopiedFiles.Num(), LOCTEXT("MigratePackages_AddToSourceControl", "Adding To Revision Control..."));
+						FScopedSlowTask SlowTask(static_cast<float>(SuccessfullyCopiedFiles.Num()), LOCTEXT("MigratePackages_AddToSourceControl", "Adding To Revision Control..."));
 						SlowTask.MakeDialog();
 
 						for(auto FileIt(SuccessfullyCopiedFiles.CreateConstIterator()); FileIt; FileIt++)
@@ -5622,7 +5622,7 @@ void UAssetToolsImpl::InitAdvancedCopyFromCopyParams(FAdvancedCopyParams CopyPar
 
 	TArray<FName> SelectedPackageNames = CopyParams.GetSelectedPackageOrFolderNames();
 
-	FScopedSlowTask SlowTask(SelectedPackageNames.Num(), LOCTEXT("AdvancedCopyPrepareSlowTask", "Preparing Files for Advanced Copy"));
+	FScopedSlowTask SlowTask(static_cast<float>(SelectedPackageNames.Num()), LOCTEXT("AdvancedCopyPrepareSlowTask", "Preparing Files for Advanced Copy"));
 	SlowTask.MakeDialog();
 
 	TArray<UAdvancedCopyCustomization*> CustomizationsToUse = CopyParams.GetCustomizationsToUse();
diff --git a/Engine/Source/Developer/AssetTools/Private/AssetTypeActions/AssetTypeActions_AnimBoneCompressionSettings.cpp b/Engine/Source/Developer/AssetTools/Private/AssetTypeActions/AssetTypeActions_AnimBoneCompressionSettings.cpp
--- a/Engine/Source/Developer/AssetTools/Private/AssetTypeActions/AssetTypeActions_AnimBoneCompressionSettings.cpp
+++ b/Engine/Source/Developer/AssetTools/Private/AssetTypeActions/AssetTypeActions_AnimBoneCompressionSettings.cpp
@@ -102,7 +102,7 @@ void FAssetTypeActions_AnimBoneCompressionSettings::ExecuteCompression(TWeakObje
 	}
 
 	const FText StatusText = FText::Format(LOCTEXT("AnimBoneCompressionSettings_Compressing", "Compressing '{0}' animations"), FText::AsNumber(AnimSeqsToRecompress.Num()));
-	FScopedSlowTask LoadingAnimSlowTask(AnimSeqsToRecompress.Num(), StatusText);
+	FScopedSlowTask LoadingAnimSlowTask(static_cast<float>(AnimSeqsToRecompress.Num()), StatusText);
 	LoadingAnimSlowTask.MakeDialog();
 
 	for (UAnimSequence* AnimSeq : AnimSeqsToRecompress)
diff --git a/Engine/Source/Developer/AssetTools/Private/AssetTypeActions/AssetTypeActions_AnimCurveCompressionSettings.cpp b/Engine/Source/Developer/AssetTools/Private/AssetTypeActions/AssetTypeActions_AnimCurveCompressionSettings.cpp
--- a/Engine/Source/Developer/AssetTools/Private/AssetTypeActions/AssetTypeActions_AnimCurveCompressionSettings.cpp
+++ b/Engine/Source/Developer/AssetTools/Private/AssetTypeActions/AssetTypeActions_AnimCurveCompressionSettings.cpp
@@ -101,7 +101,7 @@ void FAssetTypeActions_AnimCurveCompressionSettings::ExecuteCompression(TWeakObj
 	}
 
 	const FText StatusText = FText::Format(LOCTEXT("AnimCurveCompressionSettings_Compressing", "Compressing '{0}' animations"), FText::AsNumber(AnimSeqsToRecompress.Num()));
-	FScopedSlowTask LoadingAnimSlowTask(AnimSeqsToRecompress.Num(), StatusText);
+	FScopedSlowTask LoadingAnimSlowTask(static_cast<float>(AnimSeqsToRecompress.Num()), StatusText);
 	LoadingAnimSlowTask.MakeDialog();
 
 	for (UAnimSequence* AnimSeq : AnimSeqsToRecompress)
diff --git a/Engine/Source/Developer/AssetTools/Private/AssetTypeActions/AssetTypeActions_ForceFeedbackEffect.cpp b/Engine/Source/Developer/AssetTools/Private/AssetTypeActions/AssetTypeActions_ForceFeedbackEffect.cpp
--- a/Engine/Source/Developer/AssetTools/Private/AssetTypeActions/AssetTypeActions_ForceFeedbackEffect.cpp
+++ b/Engine/Source/Developer/AssetTools/Private/AssetTypeActions/AssetTypeActions_ForceFeedbackEffect.cpp
@@ -224,8 +224,8 @@ TSharedPtr<SWidget> FAssetTypeActions_ForceFeedbackEffect::GetThumbnailOverlay(c
 		.Visibility_Lambda(OnGetVisibilityLambda)
 		[
 			SNew(SBox)
-			.MinDesiredWidth(16)
-			.MinDesiredHeight(16)
+			.MinDesiredWidth(16.f)
+			.MinDesiredHeight(16.f)
 			[
 				SNew(SImage)
 				.Image_Lambda(OnGetDisplayBrushLambda)
diff --git a/Engine/Source/Developer/AssetTools/Private/AssetViewUtils.cpp b/Engine/Source/Developer/AssetTools/Private/AssetViewUtils.cpp
--- a/Engine/Source/Developer/AssetTools/Private/AssetViewUtils.cpp
+++ b/Engine/Source/Developer/AssetTools/Private/AssetViewUtils.cpp
@@ -168,7 +168,7 @@ bool AssetViewUtils::LoadAssetsIfNeeded(const TArray<FString>& ObjectPaths, TArr
 	{
 		// Get the maximum objects to load before displaying the slow task
 		const bool bShowProgressDialog = (UnloadedObjectPaths.Num() > GetDefault<UContentBrowserSettings>()->NumObjectsToLoadBeforeWarning) || bAtLeastOneUnloadedMap;
-		FScopedSlowTask SlowTask(UnloadedObjectPaths.Num(), LOCTEXT("LoadingObjects", "Loading Objects..."));
+		FScopedSlowTask SlowTask(static_cast<float>(UnloadedObjectPaths.Num()), LOCTEXT("LoadingObjects", "Loading Objects..."));
 		if (bShowProgressDialog)
 		{
 			SlowTask.MakeDialog();
diff --git a/Engine/Source/Developer/AssetTools/Private/SAdvancedCopyReportDialog.cpp b/Engine/Source/Developer/AssetTools/Private/SAdvancedCopyReportDialog.cpp
--- a/Engine/Source/Developer/AssetTools/Private/SAdvancedCopyReportDialog.cpp
+++ b/Engine/Source/Developer/AssetTools/Private/SAdvancedCopyReportDialog.cpp
@@ -209,7 +209,7 @@ void SAdvancedCopyReportDialog::Construct( const FArguments& InArgs, const FAdva
 				.Type(ESlateCheckBoxType::CheckBox)
 				.IsChecked(this, &SAdvancedCopyReportDialog::IsGeneratingDependencies)
 				.OnCheckStateChanged(this, &SAdvancedCopyReportDialog::ToggleGeneratingDependencies)
-				.Padding(4)
+				.Padding(4.f)
 				[
 					SNew(STextBlock)
 					.Text(LOCTEXT("GenerateDependenciesToCopy", "Generate Dependencies to Copy"))
diff --git a/Engine/Source/Editor/UnrealEd/Classes/Preferences/PersonaOptions.h b/Engine/Source/Editor/UnrealEd/Classes/Preferences/PersonaOptions.h
--- a/Engine/Source/Editor/UnrealEd/Classes/Preferences/PersonaOptions.h
+++ b/Engine/Source/Editor/UnrealEd/Classes/Preferences/PersonaOptions.h
@@ -62,9 +62,9 @@ struct FViewportConfigOptions
 
 	FViewportConfigOptions()
 		: ViewModeIndex(EViewModeIndex::VMI_Lit)
-		, ViewFOV(53.43)
+		, ViewFOV(53.43f)
 		, CameraSpeedSetting(4)
-		, CameraSpeedScalar(1.0)
+		, CameraSpeedScalar(1.0f)
 		, CameraFollowMode(EAnimationViewportCameraFollowMode::None)
 	{}
 
diff --git a/Engine/Source/Runtime/Interchange/Core/Public/Types/AttributeStorage.h b/Engine/Source/Runtime/Interchange/Core/Public/Types/AttributeStorage.h
--- a/Engine/Source/Runtime/Interchange/Core/Public/Types/AttributeStorage.h
+++ b/Engine/Source/Runtime/Interchange/Core/Public/Types/AttributeStorage.h
@@ -1088,7 +1088,7 @@ namespace UE
 			 * if FragmentedMemoryCost > AttributeStorage.Num*DefragRatio then defrag.
 			 * This is use whenever we remove attribute or change DefragRatio value.
 			 */
-			float DefragRatio = 0.1;
+			float DefragRatio = 0.1f;
 
 			/**
 			 * Mutex use when accessing or modifying the storage
