Removing the deprecated functions for the class permission list from the asset tools api.

diff --git a/Engine/Source/Developer/AssetTools/Private/AssetTools.cpp b/Engine/Source/Developer/AssetTools/Private/AssetTools.cpp
--- a/Engine/Source/Developer/AssetTools/Private/AssetTools.cpp
+++ b/Engine/Source/Developer/AssetTools/Private/AssetTools.cpp
@@ -5760,60 +5760,6 @@ TArray<UFactory*> UAssetToolsImpl::GetNewAssetFactories() const
 	return MoveTemp(Factories);
 }
 
-TSharedRef<FNamePermissionList>& UAssetToolsImpl::GetAssetClassPermissionList()
-{
-PRAGMA_DISABLE_DEPRECATION_WARNINGS
-	return GetAssetClassPermissionList(EAssetClassAction::ViewAsset);
-PRAGMA_ENABLE_DEPRECATION_WARNINGS
-	
-}
-
-TSharedRef<FNamePermissionList>& UAssetToolsImpl::GetAssetClassPermissionList(EAssetClassAction AssetClassAction)
-{
-	// Convert path list to name list. Slow and unefficient but keeping it for compatibility (this function is deprecated)
-	if (AssetClassAction < EAssetClassAction::AllAssetActions)
-	{
-		TSharedRef<FPathPermissionList> List = AssetClassPermissionList[(int32)AssetClassAction];
-		TSharedRef<FNamePermissionList> Result = MakeShared<FNamePermissionList>();
-		if (List->IsDenyListAll())
-		{
-			Result->AddDenyListAll(NAME_None);
-		}
-
-		for (const TPair<FString, FPermissionListOwners>& It : List->GetDenyList())
-		{
-			for (const FName OwnerName : It.Value)
-			{
-				Result->AddDenyListItem(OwnerName, FName(*FPackageName::GetShortName(*It.Key)));
-			}
-		}
-
-		for (const TPair<FString, FPermissionListOwners>& It : List->GetAllowList())
-		{
-			FName Name = *FPackageName::GetShortName(*It.Key);
-			if (Result->PassesFilter(Name))
-			{
-				for (const FName OwnerName : It.Value)
-				{
-					Result->AddAllowListItem(OwnerName, Name);
-				}
-			}
-		}
-
-		// Block everything if none of the AllowList paths passed
-		if (Result->GetAllowList().Num() == 0)
-		{
-			Result->AddDenyListAll(NAME_None);
-		}
-
-		AssetClassPermissionList_DEPRECATED = Result;
-		return AssetClassPermissionList_DEPRECATED;
-	}
-	
-	static TSharedRef<FNamePermissionList> Empty = MakeShared<FNamePermissionList>();
-	return Empty;
-}
-
 const TSharedRef<FPathPermissionList>& UAssetToolsImpl::GetAssetClassPathPermissionList(EAssetClassAction AssetClassAction) const
 {
 	if (AssetClassAction < EAssetClassAction::AllAssetActions)
diff --git a/Engine/Source/Developer/AssetTools/Private/AssetTools.h b/Engine/Source/Developer/AssetTools/Private/AssetTools.h
--- a/Engine/Source/Developer/AssetTools/Private/AssetTools.h
+++ b/Engine/Source/Developer/AssetTools/Private/AssetTools.h
@@ -133,10 +133,6 @@ public:
 	virtual void ConvertVirtualTextures(const TArray<UTexture2D*>& Textures, bool bConvertBackToNonVirtual, const TArray<UMaterial*>* RelatedMaterials = nullptr) const override;
 	virtual bool IsAssetClassSupported(const UClass* AssetClass) const override;
 	virtual TArray<UFactory*> GetNewAssetFactories() const override;
-	UE_DEPRECATED(5.1, "Class names are now represented by path names. Please use GetAssetClassPathPermissionList.")
-	virtual TSharedRef<FNamePermissionList>& GetAssetClassPermissionList() override;
-	UE_DEPRECATED(5.1, "Class names are now represented by path names. Please use GetAssetClassPathPermissionList.")
-	TSharedRef<FNamePermissionList>& GetAssetClassPermissionList(EAssetClassAction AssetClassAction);
 	virtual const TSharedRef<FPathPermissionList>& GetAssetClassPathPermissionList(EAssetClassAction AssetClassAction) const override;
 	virtual const TSharedRef<FNamePermissionList>& GetImportExtensionPermissionList() const override;
 	virtual bool IsImportExtensionAllowed(const FString& Extension) const override;
diff --git a/Engine/Source/Developer/AssetTools/Public/IAssetTools.h b/Engine/Source/Developer/AssetTools/Public/IAssetTools.h
--- a/Engine/Source/Developer/AssetTools/Public/IAssetTools.h
+++ b/Engine/Source/Developer/AssetTools/Public/IAssetTools.h
@@ -581,10 +581,6 @@ public:
 	/** Find all supported asset factories. */
 	virtual TArray<UFactory*> GetNewAssetFactories() const = 0;
 
-	/** Get asset class permission list for the ViewAsset action */
-	UE_DEPRECATED(5.1, "Pass in an EAssetClassAction instead of nothing")
-	virtual TSharedRef<FNamePermissionList>& GetAssetClassPermissionList() = 0;
-
 	/** Get asset class permission list for content browser and other systems */
 	virtual const TSharedRef<FPathPermissionList>& GetAssetClassPathPermissionList(EAssetClassAction AssetClassAction) const = 0;
 
