Migrate dialog: Migration of a map is all or nothing. Hide __ExternalActors__ / __ExternalObjects__ from dialog when migrating OFPA maps.

#rnx

diff --git a/Engine/Source/Developer/AssetTools/Private/AssetTools.cpp b/Engine/Source/Developer/AssetTools/Private/AssetTools.cpp
--- a/Engine/Source/Developer/AssetTools/Private/AssetTools.cpp
+++ b/Engine/Source/Developer/AssetTools/Private/AssetTools.cpp
@@ -4414,12 +4414,36 @@ void UAssetToolsImpl::PerformMigratePackages(TArray<FName> PackageNamesToMigrate
 	// Find assets in non-Project Plugins
 	TSet<FName> ShouldMigratePackage;
 	bool bShouldShowEngineContent = GetDefault<UContentBrowserSettings>()->GetDisplayEngineFolder();
+	TMap<FName, TSet<FName>> PackageToExternalObjectPackages;
+		
 	{
+		FAssetRegistryModule& AssetRegistryModule = FModuleManager::Get().LoadModuleChecked<FAssetRegistryModule>(TEXT("AssetRegistry"));
+		const bool bOnlyIncludeOnDiskAssets = true;
+
 		// This is the new list to prompt for migration
 		TSet<FName> FilteredPackageNamesToMove;
 
 		for (const FName& PackageName : AllPackageNamesToMove)
 		{
+			// Associate External Packages to Level Packages
+			FString PackageNameStr = PackageName.ToString();
+			if (PackageNameStr.Contains(FPackagePath::GetExternalActorsFolderName()) || PackageNameStr.Contains(FPackagePath::GetExternalObjectsFolderName()))
+			{
+				TArray<FAssetData> Assets;
+				if (AssetRegistryModule.Get().GetAssetsByPackageName(PackageName, Assets, bOnlyIncludeOnDiskAssets))
+				{
+					for (const FAssetData& AssetData : Assets)
+					{
+						if (!AssetData.GetOptionalOuterPathName().IsNone())
+						{
+							PackageToExternalObjectPackages.FindOrAdd(FSoftObjectPath(AssetData.GetOptionalOuterPathName().ToString()).GetLongPackageFName()).Add(PackageName);
+						}
+					}
+				}
+				// Avoid adding external object packages to the FilteredPackageNamesToMove because we don't want them to show up in the Migrate dialog. Instead they are going to be migrated if their outer package gets migrated.
+				continue;
+			}
+
 			FName PackageMountPoint = FPackageName::GetPackageMountPoint(PackageName.ToString(), false);
 			EPluginLoadedFrom* Found = EnabledPluginToLoadedFrom.Find(PackageMountPoint);
 
@@ -4496,16 +4520,16 @@ void UAssetToolsImpl::PerformMigratePackages(TArray<FName> PackageNamesToMigrate
 	if(!FApp::IsUnattended() && Options.bPrompt)
 	{
 		const FText ReportMessage = LOCTEXT("MigratePackagesReportTitle", "The following assets will be migrated to another content folder.");
-		SPackageReportDialog::FOnReportConfirmed OnReportConfirmed = SPackageReportDialog::FOnReportConfirmed::CreateUObject(this, &UAssetToolsImpl::MigratePackages_ReportConfirmed, ReportPackages, DestinationPath, MoveTemp(ExcludedDependencies), Options);
+		SPackageReportDialog::FOnReportConfirmed OnReportConfirmed = SPackageReportDialog::FOnReportConfirmed::CreateUObject(this, &UAssetToolsImpl::MigratePackages_ReportConfirmed, ReportPackages, DestinationPath, MoveTemp(ExcludedDependencies), MoveTemp(PackageToExternalObjectPackages), Options);
 		SPackageReportDialog::OpenPackageReportDialog(ReportMessage, *ReportPackages.Get(), OnReportConfirmed);
 	}
 	else
 	{
-		UAssetToolsImpl::MigratePackages_ReportConfirmed(ReportPackages, DestinationPath, MoveTemp(ExcludedDependencies), Options);
+		UAssetToolsImpl::MigratePackages_ReportConfirmed(ReportPackages, DestinationPath, MoveTemp(ExcludedDependencies), MoveTemp(PackageToExternalObjectPackages), Options);
 	}
 }
 
-void UAssetToolsImpl::MigratePackages_ReportConfirmed(TSharedPtr<TArray<ReportPackageData>> PackageDataToMigrate, const FString DestinationPath, TSet<FName> ExcludedDependencies, const FMigrationOptions Options) const
+void UAssetToolsImpl::MigratePackages_ReportConfirmed(TSharedPtr<TArray<ReportPackageData>> PackageDataToMigrate, const FString DestinationPath, TSet<FName> ExcludedDependencies, TMap<FName, TSet<FName>> PackageToExternalObjectPackages, const FMigrationOptions Options) const
 {
 	FString DestinationFolder;
 	if (FApp::IsUnattended() || !DestinationPath.IsEmpty())
@@ -4610,13 +4634,25 @@ void UAssetToolsImpl::MigratePackages_ReportConfirmed(TSharedPtr<TArray<ReportPa
 
 	// Build a list of packages to handle
 	TSet<FName> AllPackageNamesToMove;
+	TArray<ReportPackageData> ExternalPackagesToMigrate;
 	for (auto PackageDataIt = PackageDataToMigrate->CreateConstIterator(); PackageDataIt; ++PackageDataIt)
 	{
 		if (PackageDataIt->bShouldMigratePackage)
 		{
-			AllPackageNamesToMove.Add(FName(PackageDataIt->Name));
+			FName PackageName(PackageDataIt->Name);
+			AllPackageNamesToMove.Add(PackageName);
+
+			if (TSet<FName>* ExternalObjectPackages = PackageToExternalObjectPackages.Find(PackageName))
+			{
+				for (const FName& ExternalObjectPackage : *ExternalObjectPackages)
+				{
+					AllPackageNamesToMove.Add(ExternalObjectPackage);
+					ExternalPackagesToMigrate.Add({ ExternalObjectPackage.ToString(), true });
+				}
+			}
 		}
 	}
+	PackageDataToMigrate->Append(MoveTemp(ExternalPackagesToMigrate));
 
 	FMessageLog MigrateLog("AssetTools");
 
diff --git a/Engine/Source/Developer/AssetTools/Private/AssetTools.h b/Engine/Source/Developer/AssetTools/Private/AssetTools.h
--- a/Engine/Source/Developer/AssetTools/Private/AssetTools.h
+++ b/Engine/Source/Developer/AssetTools/Private/AssetTools.h
@@ -180,7 +180,7 @@ private:
 	void PerformAdvancedCopyPackages(TArray<FName> SelectedPackageNames, FString TargetPath, FAdvancedCopyCompletedEvent OnCopyComplete) const;
 
 	/** Copies files after the final list was confirmed */
-	void MigratePackages_ReportConfirmed(TSharedPtr<TArray<ReportPackageData>> PackageDataToMigrate, const FString DestinationPath, TSet<FName> ExcludedDependencies, const FMigrationOptions Options) const;
+	void MigratePackages_ReportConfirmed(TSharedPtr<TArray<ReportPackageData>> PackageDataToMigrate, const FString DestinationPath, TSet<FName> ExcludedDependencies, TMap<FName, TSet<FName>> PackageToExternalObjectPackages, const FMigrationOptions Options) const;
 
 	/** Copies files after the final list was confirmed */
 	void AdvancedCopyPackages_ReportConfirmed(const FAdvancedCopyParams& CopyParam, const TArray<TMap<FString, FString>>& DestinationMap) const;
