Support asset creation of child asset classes when factory is defined for super class

diff --git a/Engine/Source/Developer/AssetTools/Private/AssetTools.cpp b/Engine/Source/Developer/AssetTools/Private/AssetTools.cpp
--- a/Engine/Source/Developer/AssetTools/Private/AssetTools.cpp
+++ b/Engine/Source/Developer/AssetTools/Private/AssetTools.cpp
@@ -1635,7 +1635,10 @@ UObject* UAssetToolsImpl::CreateAsset(const FString& AssetName, const FString& P
 		return nullptr;
 	}
 
-	if ( AssetClass && Factory && !ensure(AssetClass->IsChildOf(Factory->GetSupportedClass())) )
+	// Verify the factory supports the asset class
+	if (AssetClass && Factory && 
+		((Factory->SupportedClass != nullptr && !ensure(AssetClass->IsChildOf(Factory->GetSupportedClass()))) || 
+		 (Factory->SupportedClass == nullptr && !ensure(Factory->DoesSupportClass(AssetClass)))) )
 	{
 		FMessageDialog::Open( EAppMsgType::Ok, LOCTEXT("InvalidFactory", "The new asset wasn't created because the supplied factory does not support the supplied class.") );
 		return nullptr;
