Add more details to an ensure we are seeing
#rnx

diff --git a/Engine/Source/Developer/AssetTools/Private/AssetViewUtils.cpp b/Engine/Source/Developer/AssetTools/Private/AssetViewUtils.cpp
--- a/Engine/Source/Developer/AssetTools/Private/AssetViewUtils.cpp
+++ b/Engine/Source/Developer/AssetTools/Private/AssetViewUtils.cpp
@@ -1500,7 +1500,7 @@ int32 AssetViewUtils::GetPackageLengthForCooking(const FString& PackageName, boo
 	else if (ensureMsgf(PluginContainingAsset, TEXT("Only plugins can exist outside of the expected target path of '%s'. '%s' will not calculate an accurate result!"), *AbsoluteTargetPath, *AbsolutePathToAsset))
 	{
 		const FString AbsolutePluginRootPath = FPaths::ConvertRelativePathToFull(PluginContainingAsset->GetBaseDir());
-		if (ensure(AbsolutePathToAsset.StartsWith(AbsolutePluginRootPath, ESearchCase::CaseSensitive)))
+		if (ensureMsgf(AbsolutePathToAsset.StartsWith(AbsolutePluginRootPath, ESearchCase::CaseSensitive), TEXT("%s should start with %s"), *AbsolutePathToAsset, *AbsolutePluginRootPath))
 		{
 			AssetPathRelativeToCookRootLen -= AbsolutePluginRootPath.Len();
 			AssetPathRelativeToCookRootLen += FCString::Strlen(TEXT("Plugins/GameFeatures/XXXXXXXX-XXXX-XXXX-XXXX-XXXXXXXXXXXX")); // Use a GUID instead of the plugin name, as some external plugins cook as a GUID
