Asset Package Header Patcher.
This is infrastructure for Advanced copy to not need to deserialize an entire package in order to resave it in a new location.

diff --git a/Engine/Source/Developer/AssetTools/Private/AssetHeaderPatcher.cpp b/Engine/Source/Developer/AssetTools/Private/AssetHeaderPatcher.cpp
new file mode 100644
--- /dev/null
+++ b/Engine/Source/Developer/AssetTools/Private/AssetHeaderPatcher.cpp
@@ -0,0 +1,947 @@
+// Copyright Epic Games, Inc. All Rights Reserved.
+
+#include "AssetHeaderPatcher.h"
+
+#include "Containers/ContainersFwd.h"
+#include "Algo/AnyOf.h"
+#include "UObject/PackageFileSummary.h"
+#include "UObject/NameTypes.h"
+#include "UObject/ObjectResource.h"
+#include "AssetRegistry/IAssetRegistry.h"
+
+#include "Misc/EnumerateRange.h"
+#include "Misc/FileHelper.h"
+#include "Misc/PackageName.h"
+#include "Misc/PathViews.h"
+#include "UObject/Package.h"
+#include "Serialization/LargeMemoryReader.h"
+#include "WorldPartition/WorldPartitionActorDesc.h"
+#include "WorldPartition/WorldPartitionActorDescUtils.h"
+#include "AssetRegistry/AssetData.h"
+
+DEFINE_LOG_CATEGORY_STATIC(LogAssetHeaderPatcher, Log, All);
+
+namespace 
+{
+	// To override writing of FName's to ensure they have been patched
+	class FNamePatchingWriter final : public FArchiveProxy
+	{
+	public:
+		FNamePatchingWriter(FArchive& InAr, const TMap<FNameEntryId, int32>& InNameToIndexMap)
+			: FArchiveProxy(InAr)
+			, NameToIndexMap(InNameToIndexMap)
+		{
+		}
+		
+		virtual ~FNamePatchingWriter() 
+		{
+		}
+
+		virtual FArchive& operator<<(FName& Name) override
+		{
+			FNameEntryId EntryId = Name.GetDisplayIndex();
+			const int32* MaybeIndex = NameToIndexMap.Find(EntryId);
+
+			if (MaybeIndex == nullptr)
+			{
+				UE_LOG(LogAssetHeaderPatcher, Error, TEXT("Cannot serialize FName %s because it is not in the name table for %s"), *Name.ToString(), *GetArchiveName());
+				SetCriticalError();
+				return *this;
+			}
+
+			int32 Index = *MaybeIndex;
+			int32 Number = Name.GetNumber();
+
+			FArchive& Ar = *this;
+			Ar << Index;
+			Ar << Number;
+
+			return *this;
+		}
+
+	private:
+		const TMap<FNameEntryId, int32>& NameToIndexMap;
+	};
+
+	enum class EPatchedSection 
+	{
+		Summary,
+		NameTable,
+		ExportTable,
+		ThumbnailTable,
+		AssetRegistryData
+	};
+
+	struct FSectionData
+	{
+		EPatchedSection Section = EPatchedSection::Summary;
+		int64 Offset = 0;
+		int64 Size = 0;
+		bool bRequired = false;
+	};
+
+	enum class ESummaryOffset
+	{
+		NameTable,
+		SoftObjectPathList,
+		GatherableTextDataOffset,
+		ImportTable,
+		ExportTable,
+		DependsTable,
+		SoftPackageReferenceList,
+		SearchableNamesTable,
+		ThumbnailTable,
+		AssetRegistryData,
+		WorldTileInfoData,
+		PreloadDependency, // Should not be present - only for cooked data
+		BulkData,
+		PayloadToc
+	};
+
+	// To override MemoryReaders FName method
+	class FReadFNameAs2IntFromMemoryReader final : public FLargeMemoryReader
+	{
+	public:
+		FReadFNameAs2IntFromMemoryReader(TArray<FName>& InNameTable, const uint8* InData, const int64 Num, ELargeMemoryReaderFlags InFlags = ELargeMemoryReaderFlags::None, const FName InArchiveName = NAME_None)
+			: FLargeMemoryReader(InData, Num, InFlags, InArchiveName)
+			, NameTable(InNameTable)
+		{
+		}
+
+		// FLargeMemoryReader falls back to FMemoryArchive's imp of this method.
+		// which uses strings as the format for FName.
+		// We need the 2xint32 version when decoding the current file formats. 
+		virtual FArchive& operator<<(FName& OutName) override
+		{
+			int32 NameIndex;
+			int32 Number;
+			FArchive& Ar = *this;
+			Ar << NameIndex;
+			Ar << Number;
+
+			if (NameTable.IsValidIndex(NameIndex))
+			{
+				FNameEntryId MappedName = NameTable[NameIndex].GetDisplayIndex();
+				OutName = FName::CreateFromDisplayId(MappedName, MappedName ? Number : 0);
+			}
+			else
+			{
+				OutName = FName();
+				SetCriticalError();
+			}
+
+			return *this;
+		}
+
+		virtual FString GetArchiveName() const override 
+		{
+			return TEXT("FReadFNameAs2IntFromMemoryReader"); 
+		}
+	private:
+		TArray<FName>& NameTable;
+	};
+
+	int FindInPath(FStringView HayStack, FStringView Needle, int At = 0) 
+	{
+		check(At >= 0);
+		for (; At < HayStack.Len() - Needle.Len(); ++At)
+		{
+			if (FPathViews::Equals(FStringView(HayStack.GetData() + At, Needle.Len()), Needle))
+			{
+				return At;
+			}			
+		}
+
+		return -1;
+	}
+
+	struct FSummaryOffsetMeta
+	{
+		// NOTE: The offsets in Summary get to a max of 312 bytes.
+		// So we could drop this to a uint16 but that is probably overkill at this point.
+		uint32 Offset : 31;
+		uint32 bIs64Bit : 1;
+
+		int64 Value(FPackageFileSummary& Summary) const
+		{
+			intptr_t Ptr = reinterpret_cast<intptr_t>(&Summary) + Offset;
+			if (bIs64Bit)
+			{
+				return *reinterpret_cast<int64*>(Ptr);
+			}
+			else
+			{
+				return *reinterpret_cast<int32*>(Ptr);
+			}
+		}
+
+		void PatchOffsetValue(FPackageFileSummary& Summary, int64 Value) const
+		{
+			intptr_t Ptr = reinterpret_cast<intptr_t>(&Summary) + Offset;
+			if (bIs64Bit)
+			{
+				int64& Dst = *reinterpret_cast<int64*>(Ptr);
+				Dst += Value;
+			}
+			else
+			{
+				int32& Dst = *reinterpret_cast<int32*>(Ptr);
+				*reinterpret_cast<int32*>(Ptr) = IntCastChecked<int32>((int64)Dst + Value);
+			}
+		}
+	};
+
+	const FSummaryOffsetMeta OffsetTable[] = {
+
+	#define UE_POPULATE_OFFSET_INFO(NAME)					\
+		(uint32)STRUCT_OFFSET(FPackageFileSummary, NAME),	\
+		std::is_same_v<decltype(((FPackageFileSummary*)0)->NAME), int64>
+
+			{ UE_POPULATE_OFFSET_INFO(NameOffset) },
+			{ UE_POPULATE_OFFSET_INFO(SoftObjectPathsOffset) },
+			{ UE_POPULATE_OFFSET_INFO(GatherableTextDataOffset) },
+			{ UE_POPULATE_OFFSET_INFO(ImportOffset) },
+			{ UE_POPULATE_OFFSET_INFO(ExportOffset) },
+			{ UE_POPULATE_OFFSET_INFO(DependsOffset) },
+			{ UE_POPULATE_OFFSET_INFO(SoftPackageReferencesOffset) },
+			{ UE_POPULATE_OFFSET_INFO(SearchableNamesOffset) },
+			{ UE_POPULATE_OFFSET_INFO(ThumbnailTableOffset) },
+			{ UE_POPULATE_OFFSET_INFO(AssetRegistryDataOffset) },
+			{ UE_POPULATE_OFFSET_INFO(BulkDataStartOffset) },
+			{ UE_POPULATE_OFFSET_INFO(WorldTileInfoDataOffset) },
+			{ UE_POPULATE_OFFSET_INFO(PreloadDependencyOffset) },
+			{ UE_POPULATE_OFFSET_INFO(PayloadTocOffset) },
+
+#undef UE_POPULATE_OFFSET_INFO
+	};
+
+	void PatchSummaryOffsets(FPackageFileSummary& Dst, int64 OffsetFrom, int64 OffsetDelta)
+	{
+		if (!OffsetDelta)
+		{
+			return;
+		}
+
+		for (const FSummaryOffsetMeta& OffsetData : OffsetTable)
+		{
+			if (OffsetData.Value(Dst) > OffsetFrom)
+			{
+				OffsetData.PatchOffsetValue(Dst, OffsetDelta);
+			}
+		}
+	};
+
+	FAssetDataTagMap MakeTagMap(const TArray<UE::AssetRegistry::FDeserializeTagData>& TagData)
+	{
+		FAssetDataTagMap Out;
+		Out.Reserve(TagData.Num());
+		for (const UE::AssetRegistry::FDeserializeTagData& Tag : TagData)
+		{
+			if (!Tag.Key.IsEmpty() && !Tag.Value.IsEmpty())
+			{
+				Out.Add(*Tag.Key, Tag.Value);
+			}
+		}
+
+		return Out;
+	}
+
+}
+
+// The information we need in the task to do patching.
+class FAssetHeaderPatcherInner
+{
+public:
+	using EResult = FAssetHeaderPatcher::EResult;
+
+	FAssetHeaderPatcherInner(FString InSrcAsset, FString InDstAsset, TMap<FString, FString> InSearchAndReplace, FArchive* InDstArchive = nullptr)
+		: SrcAsset(MoveTemp(InSrcAsset))
+		, DstAsset(MoveTemp(InDstAsset))
+		, DstArchive(InDstArchive)
+		, SearchAndReplace(MoveTemp(InSearchAndReplace))
+	{
+	}
+
+	bool DoPatch(FString& InOutString);
+	bool DoPatch(FName& InOutName);
+	bool DoPatch(FTopLevelAssetPath& InOutPath);
+
+	FAssetHeaderPatcher::EResult PatchHeader();
+	FAssetHeaderPatcher::EResult PatchHeader_Deserialize();
+	void PatchHeader_PatchSections();
+	FAssetHeaderPatcher::EResult PatchHeader_WriteDestinationFile();
+
+	FString SrcAsset;
+	FString DstAsset;
+	FArchive* DstArchive = nullptr;
+	TUniquePtr<FArchive> DstArchiveOwner;
+	TMap<FString, FString> SearchAndReplace;
+	TArray64<uint8> SrcBuffer;
+
+	struct FHeaderInformation
+	{
+		int64 SummarySize = -1;
+		int64 NameTableSize = -1;
+		int64 SoftObjectPathListSize = -1;
+		int64 SoftPackageReferencesListSize = -1;
+		int64 ImportTableSize = -1;
+		int64 ExportTableSize = -1;
+		int64 ThumbnailTableSize = -1;
+		int64 AssetRegistryDataSize = -1;
+		int64 PackageTrailerSize = -1;
+	};
+
+	struct FThumbnailEntry
+	{
+		FString ObjectShortClassName;
+		FString ObjectPathWithoutPackageName;
+		int32 FileOffset = 0;
+	};
+
+	FHeaderInformation HeaderInformation;
+	FPackageFileSummary Summary;
+	TArray<FName> NameTable;
+	TMap<FNameEntryId, int32> NameToIndexMap;
+	TArray<FObjectImport> ImportTable;
+	TArray<FObjectExport> ExportTable;
+	TArray<FThumbnailEntry> ThumbnailTable;
+
+	// Asset registry data information
+	struct FAsetRegistryObjectData
+	{
+		UE::AssetRegistry::FDeserializeObjectPackageData ObjectData;
+		TArray<UE::AssetRegistry::FDeserializeTagData> TagData;
+	};
+
+	struct FAsetRegistryData
+	{
+		int64 SectionSize = -1;
+		UE::AssetRegistry::FDeserializePackageData PkgData;
+		TArray<FAsetRegistryObjectData> ObjectData;
+	};
+	FAsetRegistryData AsetRegistryData;
+};
+
+
+UE::Tasks::TTask<FAssetHeaderPatcher::EResult> FAssetHeaderPatcher::Start(FString InSrcAsset, FString InDstAsset,
+	TMap<FString, FString> InSearchAndReplace)
+{
+	TUniquePtr<FAssetHeaderPatcherInner> Self = MakeUnique<FAssetHeaderPatcherInner>(MoveTemp(InSrcAsset), MoveTemp(InDstAsset), MoveTemp(InSearchAndReplace));
+
+	return UE::Tasks::Launch(UE_SOURCE_LOCATION, [Self = MoveTemp(Self)]() -> FAssetHeaderPatcher::EResult
+		{
+			if (!FFileHelper::LoadFileToArray(Self->SrcBuffer, *Self->SrcAsset))
+			{
+				UE_LOG(LogAssetHeaderPatcher, Error, TEXT("Failed to load %s"), *Self->SrcAsset);
+				return FAssetHeaderPatcherInner::EResult::ErrorFailedToLoadSourceAsset;
+			} 
+			else
+			{
+				return Self->PatchHeader();
+			}
+		});
+}
+
+UE::Tasks::TTask<FAssetHeaderPatcher::EResult> FAssetHeaderPatcher::Start(TUniquePtr<FArchive> InReader, FString InDstAsset,
+	TMap<FString, FString> InSearchAndReplace)
+{
+	if (!InReader || InReader->IsError())
+	{
+		UE_LOG(LogAssetHeaderPatcher, Error, TEXT("Failed to load %s (archive is invalid)"), InReader ? *InReader->GetArchiveName() : TEXT("Unknown"));
+
+		return UE::Tasks::Launch(UE_SOURCE_LOCATION, []() -> EResult { return EResult::ErrorFailedToLoadSourceAsset; });
+	}
+
+	TUniquePtr<FAssetHeaderPatcherInner> Self = MakeUnique<FAssetHeaderPatcherInner>(InReader->GetArchiveName(), MoveTemp(InDstAsset), MoveTemp(InSearchAndReplace));
+
+	return UE::Tasks::Launch(UE_SOURCE_LOCATION, [Self = MoveTemp(Self), Reader = MoveTemp(InReader)]() -> EResult
+		{
+			Reader->Seek(0);
+			Self->SrcBuffer.SetNumUninitialized(Reader->TotalSize());
+			Reader->Serialize(Self->SrcBuffer.GetData(), Self->SrcBuffer.Num());
+
+			if (Reader->IsError())
+			{
+				UE_LOG(LogAssetHeaderPatcher, Error, TEXT("Failed to load %s"), *Self->SrcAsset);
+				return FAssetHeaderPatcherInner::EResult::ErrorFailedToLoadSourceAsset;
+			}
+			else
+			{
+				return Self->PatchHeader();
+			}
+		});
+}
+
+FAssetHeaderPatcher::EResult FAssetHeaderPatcher::Test_DoPatch(FArchive& InReader, FArchive& InWriter,
+	TMap<FString, FString> InSearchAndReplace)
+{
+	FAssetHeaderPatcherInner Inner(InReader.GetArchiveName(), InWriter.GetArchiveName(), MoveTemp(InSearchAndReplace), &InWriter);
+
+	InReader.Seek(0);
+	Inner.SrcBuffer.SetNumUninitialized(InReader.TotalSize());
+	InReader.Serialize(Inner.SrcBuffer.GetData(), Inner.SrcBuffer.Num());
+
+	if (InReader.IsError())
+	{
+		UE_LOG(LogAssetHeaderPatcher, Error, TEXT("Failed to load %s"), *Inner.SrcAsset);
+		return FAssetHeaderPatcherInner::EResult::ErrorFailedToLoadSourceAsset;
+	}
+	else
+	{
+		return Inner.PatchHeader();
+	}
+}
+
+FAssetHeaderPatcher::EResult FAssetHeaderPatcherInner::PatchHeader()
+{
+	FAssetHeaderPatcher::EResult Result = PatchHeader_Deserialize();
+	if (Result != EResult::None)
+	{
+		return Result;
+	}
+
+	PatchHeader_PatchSections();
+
+	return PatchHeader_WriteDestinationFile();
+}
+
+FAssetHeaderPatcher::EResult FAssetHeaderPatcherInner::PatchHeader_Deserialize()
+{
+	FReadFNameAs2IntFromMemoryReader MemAr(NameTable, SrcBuffer.GetData(), SrcBuffer.Num());
+
+	MemAr << Summary;
+	HeaderInformation.SummarySize = MemAr.Tell();
+
+	// set version numbers so components branch correctly
+	MemAr.SetUEVer(Summary.GetFileVersionUE());
+	MemAr.SetLicenseeUEVer(Summary.GetFileVersionLicenseeUE());
+	MemAr.SetEngineVer(Summary.SavedByEngineVersion);
+	MemAr.SetCustomVersions(Summary.GetCustomVersionContainer());
+
+	if (Summary.DataResourceOffset > 0)
+	{
+		// Should only be set in cooked data. If that changes, we need to add code to patch it
+		UE_LOG(LogAssetHeaderPatcher, Error, TEXT("Asset %s has an unexpected DataResourceOffset"), *SrcAsset);
+		return EResult::ErrorUnexpectedSectionOrder;
+	}
+
+	if (Summary.NameCount > 0)
+	{
+		MemAr.Seek(Summary.NameOffset);
+		NameTable.Reserve(Summary.NameCount);
+		for (int32 NameMapIdx = 0; NameMapIdx < Summary.NameCount; ++NameMapIdx)
+		{
+			FNameEntrySerialized NameEntry(ENAME_LinkerConstructor);
+			MemAr << NameEntry;
+			NameTable.Add(FName(NameEntry));
+		}
+
+		HeaderInformation.NameTableSize = MemAr.Tell() - HeaderInformation.SummarySize;
+	}
+
+#define UE_CHECK_AND_SET_ERROR_AND_RETURN(EXP)	\
+	do											\
+	{											\
+		if (EXP)								\
+		{										\
+			UE_LOG(LogAssetHeaderPatcher, Log, TEXT("Asset %s fails %s"), *SrcAsset, TEXT(#EXP));	\
+			return EResult::ErrorBadOffset;		\
+		}										\
+	}											\
+	while(0)
+
+	if (Summary.ImportCount > 0)
+	{
+		UE_CHECK_AND_SET_ERROR_AND_RETURN(Summary.ImportOffset >= Summary.TotalHeaderSize);
+		UE_CHECK_AND_SET_ERROR_AND_RETURN(Summary.ImportOffset < 0);
+
+		MemAr.Seek(Summary.ImportOffset);
+		ImportTable.Reserve(Summary.ImportCount);
+		for (int32 ImportIndex = 0; ImportIndex < Summary.ImportCount; ++ImportIndex)
+		{
+			FObjectImport& Import = ImportTable.Emplace_GetRef();
+			MemAr << Import;
+		}
+
+		HeaderInformation.ImportTableSize = MemAr.Tell() - Summary.ImportOffset;
+	}
+	else 
+	{
+		HeaderInformation.ImportTableSize = 0;
+	}
+
+	if (Summary.ExportCount > 0)
+	{
+		UE_CHECK_AND_SET_ERROR_AND_RETURN(Summary.ExportOffset >= Summary.TotalHeaderSize);
+		UE_CHECK_AND_SET_ERROR_AND_RETURN(Summary.ExportOffset < 0);
+
+		MemAr.Seek(Summary.ExportOffset);
+		ExportTable.Reserve(Summary.ExportCount);
+		for (int32 ExportIndex = 0; ExportIndex < Summary.ExportCount; ++ExportIndex)
+		{
+			FObjectExport& Export = ExportTable.Emplace_GetRef();
+			MemAr << Export;
+		}
+
+		HeaderInformation.ExportTableSize = MemAr.Tell() - Summary.ExportOffset;
+	}
+	else
+	{
+		HeaderInformation.ExportTableSize = 0;
+	}
+
+#undef UE_CHECK_AND_SET_ERROR_AND_RETURN
+
+	if (Summary.ThumbnailTableOffset)
+	{
+		MemAr.Seek(Summary.ThumbnailTableOffset);
+
+		int32 ThumbnailCount = 0;
+		MemAr << ThumbnailCount;
+
+		ThumbnailTable.Reserve(ThumbnailCount);
+		for (int32 Index = 0; Index < ThumbnailCount; ++Index)
+		{
+			FThumbnailEntry& Entry = ThumbnailTable.Emplace_GetRef();
+			MemAr << Entry.ObjectShortClassName;
+			MemAr << Entry.ObjectPathWithoutPackageName;
+			MemAr << Entry.FileOffset;
+		}
+
+		HeaderInformation.ThumbnailTableSize = MemAr.Tell() - Summary.ThumbnailTableOffset;
+	}
+
+	// Load AR data
+	if (Summary.AssetRegistryDataOffset)
+	{
+		MemAr.Seek(Summary.AssetRegistryDataOffset);
+
+		UE::AssetRegistry::EReadPackageDataMainErrorCode ErrorCode;
+		if (!AsetRegistryData.PkgData.DoSerialize(MemAr, Summary, ErrorCode))
+		{
+			UE_LOG(LogAssetHeaderPatcher, Error, TEXT("Failed to deserialize asset registry data for %s"), *SrcAsset);
+			return EResult::ErrorFailedToDeserializeSourceAsset;
+		}
+
+		AsetRegistryData.ObjectData.Reserve(AsetRegistryData.PkgData.ObjectCount);
+		for (int32 i = 0; i < AsetRegistryData.PkgData.ObjectCount; ++i)
+		{
+			FAsetRegistryObjectData& ObjData = AsetRegistryData.ObjectData.Emplace_GetRef();
+			if (!ObjData.ObjectData.DoSerialize(MemAr, ErrorCode))
+			{
+				UE_LOG(LogAssetHeaderPatcher, Error, TEXT("Failed to deserialize asset registry data for %s"), *SrcAsset);
+				return EResult::ErrorFailedToDeserializeSourceAsset;
+			}
+
+			ObjData.TagData.Reserve(ObjData.ObjectData.TagCount);
+			for (int32 j = 0; j < ObjData.ObjectData.TagCount; ++j)
+			{
+				if (!ObjData.TagData.Emplace_GetRef().DoSerialize(MemAr, ErrorCode))
+				{
+					UE_LOG(LogAssetHeaderPatcher, Error, TEXT("Failed to deserialize asset registry data for %s"), *SrcAsset);
+					return EResult::ErrorFailedToDeserializeSourceAsset;
+				}
+			}
+		}
+
+		AsetRegistryData.SectionSize = MemAr.Tell() - Summary.AssetRegistryDataOffset;
+	}
+
+	return EResult::Success;
+}
+
+bool FAssetHeaderPatcherInner::DoPatch(FString& InOutString)
+{
+	for (const TPair<FString, FString>& KeyValue : SearchAndReplace)
+	{
+		if (FPathViews::Equals(InOutString, KeyValue.Key))
+		{
+			InOutString = KeyValue.Value;
+			return true;
+		}
+	}
+
+	return false;
+}
+
+bool FAssetHeaderPatcherInner::DoPatch(FName& InOutName)
+{
+	FString Value = InOutName.GetPlainNameString();
+	bool bPatching = DoPatch(Value);
+	if (bPatching)
+	{
+		// Use the same Number as the original Name for consistency.
+		// Otherwise different files with the same name, generate different numbers, and this breaks linking.
+		InOutName = FName(Value, InOutName.GetNumber());
+	}
+	return bPatching;
+}
+
+bool FAssetHeaderPatcherInner::DoPatch(FTopLevelAssetPath& InOutPath)
+{
+	FString Value = InOutPath.ToString();
+	bool bPatching = DoPatch(Value);
+	if (bPatching)
+	{
+		InOutPath = Value;
+	}
+	return bPatching;
+}
+
+void FAssetHeaderPatcherInner::PatchHeader_PatchSections()
+{
+	DoPatch(Summary.PackageName);
+
+	for (FName& Name : NameTable)
+	{
+		DoPatch(Name);
+	}
+
+	for (FObjectImport& Import : ImportTable)
+	{
+		DoPatch(Import.ObjectName);
+#if WITH_EDITORONLY_DATA
+		DoPatch(Import.OldClassName);
+#endif
+
+		DoPatch(Import.ClassPackage);
+		DoPatch(Import.ClassName);
+#if WITH_EDITORONLY_DATA
+		DoPatch(Import.PackageName);
+#endif
+	}
+
+	for (FObjectExport& Export : ExportTable)
+	{
+		DoPatch(Export.ObjectName);
+#if WITH_EDITORONLY_DATA
+		DoPatch(Export.OldClassName);
+#endif
+	}
+
+	// Asset Register Data
+	for (FAsetRegistryObjectData& ObjData : AsetRegistryData.ObjectData)
+	{
+		DoPatch(ObjData.ObjectData.ObjectPath);
+		DoPatch(ObjData.ObjectData.ObjectClassName);
+
+		bool bHasActor = false;
+
+		for (UE::AssetRegistry::FDeserializeTagData& TagData : ObjData.TagData)
+		{
+			bHasActor |= (TagData.Key == FWorldPartitionActorDescUtils::ActorMetaDataClassTagName()
+				       || TagData.Key == FWorldPartitionActorDescUtils::ActorMetaDataTagName());
+			
+			DoPatch(TagData.Value);
+		}
+
+		if (bHasActor)
+		{ 
+			FGCScopeGuard GCGuard;
+
+			const FString LongPackageName(SrcAsset);
+			const FString ObjectPath(ObjData.ObjectData.ObjectPath);
+			const FTopLevelAssetPath AssetClassPathName(ObjData.ObjectData.ObjectClassName);
+			const FAssetDataTagMap Tags(MakeTagMap(ObjData.TagData));
+			const FAssetData AssetData(LongPackageName, ObjectPath, AssetClassPathName, Tags);
+
+			if (TUniquePtr<FWorldPartitionActorDesc> ActorDesc = FWorldPartitionActorDescUtils::GetActorDescriptorFromAssetData(AssetData))
+			{
+				bool bPatchedActorDesc = DoPatch(ActorDesc->BaseClass);
+
+				if (ActorDesc->bIsUsingDataLayerAsset)
+				{
+					for (FName& DataLayerAssetPath : ActorDesc->DataLayers)
+					{
+						bPatchedActorDesc |= DoPatch(DataLayerAssetPath);
+					}
+				}
+
+				if (bPatchedActorDesc)
+				{
+					for (UE::AssetRegistry::FDeserializeTagData& Tag : ObjData.TagData)
+					{
+						if (Tag.Key == FWorldPartitionActorDescUtils::ActorMetaDataTagName())
+						{
+							Tag.Value = FWorldPartitionActorDescUtils::GetAssetDataFromActorDescriptor(ActorDesc);
+							break;
+						}
+					}
+				}
+			}
+		}
+	}
+
+	// Build Name Path table
+	NameToIndexMap.Reserve(NameTable.Num());
+	for (TConstEnumerateRef<FName> Name : EnumerateRange(NameTable))
+	{
+		NameToIndexMap.Add(Name->GetDisplayIndex()) = Name.GetIndex();
+	}
+}
+
+FAssetHeaderPatcher::EResult FAssetHeaderPatcherInner::PatchHeader_WriteDestinationFile()
+{
+	// Serialize modified sections and reconstruct the file	
+	// Original offsets and sizes of any sections that will be patched
+	//	Tag											Offset									Size						bRequired
+	const FSectionData SourceSections[] = {
+		{ EPatchedSection::Summary,					0,										HeaderInformation.SummarySize,			true },
+		{ EPatchedSection::NameTable,				Summary.NameOffset,						HeaderInformation.NameTableSize,		true },
+		{ EPatchedSection::ExportTable,				Summary.ExportOffset,					HeaderInformation.ExportTableSize,		true },
+		{ EPatchedSection::ThumbnailTable,			Summary.ThumbnailTableOffset,			HeaderInformation.ThumbnailTableSize,	false },
+		{ EPatchedSection::AssetRegistryData,		Summary.AssetRegistryDataOffset,		AsetRegistryData.SectionSize,			true },
+	};
+
+	const int32 SourceTotalHeaderSize = Summary.TotalHeaderSize;
+
+	// Ensure the sections are in the expected order.
+	for (int SectionIdx = 1; SectionIdx < UE_ARRAY_COUNT(SourceSections); ++SectionIdx)
+	{
+		const FSectionData& SourceSection = SourceSections[SectionIdx];
+		const FSectionData& PrevSection = SourceSections[SectionIdx - 1];
+
+		// Verify sections are ordered as expected
+		if (SourceSection.Offset < 0 || (SourceSection.bRequired && (SourceSection.Offset < PrevSection.Offset)))
+		{
+			UE_LOG(LogAssetHeaderPatcher, Error, TEXT("Unexpected section order for %s (%d %zd < %zd) "), *SrcAsset, SectionIdx, SourceSection.Offset, PrevSection.Offset);
+			return EResult::ErrorUnexpectedSectionOrder;
+		}
+	}
+
+	// Ensure the required sections have data
+	for (const FSectionData& SourceSection : SourceSections)
+	{
+		// skip processing empty non required chunks.
+		if (SourceSection.bRequired && SourceSection.Size <= 0)
+		{
+			UE_LOG(LogAssetHeaderPatcher, Error, TEXT("Unexpected section order for %s"), *SrcAsset);
+			return EResult::ErrorEmptyRequireSection;
+		}
+	}
+
+	// Create the destination file if not open already
+	if (!DstArchive)
+	{
+		TUniquePtr<FArchive> FileWriter(IFileManager::Get().CreateFileWriter(*DstAsset, FILEWRITE_EvenIfReadOnly));
+		if (!FileWriter)
+		{
+			UE_LOG(LogAssetHeaderPatcher, Error, TEXT("Failed to open %s for write"), *DstAsset);
+			return EResult::ErrorFailedToOpenDestinationFile;
+		}
+		DstArchiveOwner = MoveTemp(FileWriter);
+		DstArchive = DstArchiveOwner.Get();
+	}
+
+	FNamePatchingWriter Writer(*DstArchive, NameToIndexMap);
+
+	// set version numbers so components branch correctly
+	Writer.SetUEVer(Summary.GetFileVersionUE());
+	Writer.SetLicenseeUEVer(Summary.GetFileVersionLicenseeUE());
+	Writer.SetEngineVer(Summary.SavedByEngineVersion);
+	Writer.SetCustomVersions(Summary.GetCustomVersionContainer());
+
+	int64 LastSectionEndedAt = 0;
+	
+	for (int SectionIdx = 0; SectionIdx < UE_ARRAY_COUNT(SourceSections); ++SectionIdx)
+	{
+		const FSectionData& SourceSection = SourceSections[SectionIdx];
+
+		// skip processing empty non required chunks.
+		// really the only option section is the Thumbnails, and its annoying its ion the middle.
+		if (!SourceSection.bRequired && SourceSection.Size <= 0) 
+		{
+			continue;
+		}
+
+		// copy the blob from the end of the last section, to the start of this one
+		if (LastSectionEndedAt)
+		{
+			int64 SizeToCopy = SourceSection.Offset - LastSectionEndedAt;
+			checkf(SizeToCopy >= 0, TEXT("Section %d of %s\n%zd -> %zd %zd"), SectionIdx, *SrcAsset, SourceSection.Offset, LastSectionEndedAt, SizeToCopy);
+			Writer.Serialize(SrcBuffer.GetData() + LastSectionEndedAt, SizeToCopy);
+		}
+		LastSectionEndedAt = SourceSection.Offset + SourceSection.Size;
+
+		// Serialize the current patched section and patch summary offsets
+		switch (SourceSection.Section)
+		{
+			case EPatchedSection::Summary:
+			{
+				// We will write the Summary twice.
+				// The first is so we get its new size (if the name was changed in patching)
+				// The second is done after the loop, to patch up all the offsets.
+				check(Writer.Tell() == 0);
+				Writer << Summary;
+				const int64 SummarySize = Writer.Tell();
+				const int64 Delta = SummarySize - SourceSection.Size;
+				PatchSummaryOffsets(Summary, 0, Delta);
+				Summary.TotalHeaderSize += (int32)Delta;
+
+				break;
+			}
+
+			case EPatchedSection::NameTable:
+			{
+				const int64 NameTableStartOffset = Writer.Tell();
+				for (FName& Name : NameTable)
+				{
+					const FNameEntry* Entry = FName::GetEntry(Name.GetDisplayIndex());
+					check(Entry);
+					Entry->Write(Writer);
+				}
+				const int64 NameTableSize = Writer.Tell() - NameTableStartOffset;
+				const int64 Delta = NameTableSize - SourceSection.Size;
+				PatchSummaryOffsets(Summary, NameTableStartOffset, Delta);
+				Summary.TotalHeaderSize += (int32)Delta;
+				check(Summary.NameCount == NameTable.Num());
+				check(Summary.NameOffset == NameTableStartOffset);
+
+				break;
+			}
+
+			case EPatchedSection::ExportTable:
+			{
+				// The export table offsets aren't correct yet.
+				// Once we know them, we will seek back and write it a second time.
+				const int64 ExportTableStartOffset = Writer.Tell();
+				for (FObjectExport& Export : ExportTable)
+				{
+					Writer << Export;
+				}
+				const int64 ExportTableSize = Writer.Tell() - ExportTableStartOffset;
+				const int64 Delta = ExportTableSize - SourceSection.Size;
+				PatchSummaryOffsets(Summary, ExportTableStartOffset, Delta);
+
+				checkf(ExportTableSize == SourceSection.Size, TEXT("%d == %d"), (int)ExportTableSize, (int)SourceSection.Size); // We only patch export table offsets, we should not be patching size
+				checkf(Summary.ExportCount == ExportTable.Num(), TEXT("%d == %d"), Summary.ExportCount, ExportTable.Num());
+				checkf(Summary.ExportOffset == ExportTableStartOffset, TEXT("%d == %d"), Summary.ExportOffset, ExportTableStartOffset);
+
+				break;
+			}
+
+			case EPatchedSection::ThumbnailTable:
+			{
+				// Luckily the strings in the thumbnail table don't include the package, so only the offsets need to be patched
+				const int64 ThumbnailTableStartOffset = Writer.Tell();
+				const int64 ThumbnailTableDeltaOffset = ThumbnailTableStartOffset - SourceSection.Offset;
+				int32 ThumbnailCount = ThumbnailTable.Num();
+				Writer << ThumbnailCount;
+				for (FThumbnailEntry& Entry : ThumbnailTable)
+				{
+					Writer << Entry.ObjectShortClassName;
+					Writer << Entry.ObjectPathWithoutPackageName;
+					Entry.FileOffset += (int32)ThumbnailTableDeltaOffset; // Thumbnail payloads immediately follow the table, so we can just apply the section offset delta here
+					Writer << Entry.FileOffset;
+				}
+				const int64 ThumbnailTableSize = Writer.Tell() - ThumbnailTableStartOffset;
+				checkf(ThumbnailTableStartOffset == Summary.ThumbnailTableOffset, TEXT("%zd == %zd"), ThumbnailTableStartOffset, Summary.ThumbnailTableOffset);
+				check(ThumbnailTableSize == SourceSection.Size); // We only patch thumbnail table offsets, we should not be patching size
+
+				break;
+			}
+
+			case EPatchedSection::AssetRegistryData:
+			{
+				const int64 AsetRegistryDataStartOffset = Writer.Tell();
+				checkf(AsetRegistryDataStartOffset == Summary.AssetRegistryDataOffset, TEXT("%zd == %zd"), AsetRegistryDataStartOffset, Summary.AssetRegistryDataOffset);
+
+				// Manually write this back out, there isn't a nicely factored function to call for this
+				if (AsetRegistryData.PkgData.DependencyDataOffset != INDEX_NONE)
+				{
+					Writer << AsetRegistryData.PkgData.DependencyDataOffset;
+				}
+				Writer << AsetRegistryData.PkgData.ObjectCount;
+
+				check(AsetRegistryData.PkgData.ObjectCount == AsetRegistryData.ObjectData.Num());
+				for (FAsetRegistryObjectData& ObjData : AsetRegistryData.ObjectData)
+				{
+					Writer << ObjData.ObjectData.ObjectPath;
+					Writer << ObjData.ObjectData.ObjectClassName;
+					Writer << ObjData.ObjectData.TagCount;
+
+					check(ObjData.ObjectData.TagCount == ObjData.TagData.Num());
+					for (UE::AssetRegistry::FDeserializeTagData& TagData : ObjData.TagData)
+					{
+						Writer << TagData.Key;
+						Writer << TagData.Value;
+					}
+				}
+
+				const int64 AsetRegistryDataSize = Writer.Tell() - AsetRegistryDataStartOffset;
+				const int64 Delta = AsetRegistryDataSize - SourceSection.Size;
+				PatchSummaryOffsets(Summary, AsetRegistryDataStartOffset, Delta);
+				Summary.TotalHeaderSize += (int32)Delta;
+
+				if (AsetRegistryData.PkgData.DependencyDataOffset != INDEX_NONE)
+				{
+					// DependencyDataOffset is not relative but points to just after the rest of the AR data
+					// We will seek back and write this later
+					const int64 DependencyDataDelta = AsetRegistryDataStartOffset - SourceSection.Offset + Delta;
+					AsetRegistryData.PkgData.DependencyDataOffset += DependencyDataDelta;
+				}
+
+				break;
+			}
+
+			default:
+				UE_LOG(LogAssetHeaderPatcher, Error, TEXT("Unexpected section for %s"), *SrcAsset);
+				return EResult::ErrorUnkownSection;
+		}
+
+		if (Writer.IsError())
+		{
+			UE_LOG(LogAssetHeaderPatcher, Error, TEXT("Failed to write to %s"), *DstAsset);
+			return EResult::ErrorFailedToWriteToDestinationFile;
+		}
+	}
+
+	{   // copy the last blob
+		int64 SizeToCopy = SrcBuffer.Num() - LastSectionEndedAt;
+		checkf(SizeToCopy >= 0, TEXT("Section last of %s\n%zd -> %zd %zd"), *SrcAsset, SrcBuffer.Num(), LastSectionEndedAt, SizeToCopy);
+		Writer.Serialize(SrcBuffer.GetData() + LastSectionEndedAt, SizeToCopy);
+	}
+
+	if (Writer.IsError())
+	{
+		UE_LOG(LogAssetHeaderPatcher, Error, TEXT("Failed to write to %s"), *DstAsset);
+		return EResult::ErrorFailedToWriteToDestinationFile;
+	}
+
+	// Re-write summary with patched offsets
+	Writer.Seek(0);
+	Writer << Summary;
+
+	{
+		// Re-write export table with patched offsets
+		// Patch Export table offsets now that we have patched all the header sections
+		Writer.Seek(Summary.ExportOffset);
+		const int64 ExportOffsetDelta = static_cast<int64>(Summary.TotalHeaderSize) - SourceTotalHeaderSize;
+		for (FObjectExport& Export : ExportTable)
+		{
+			Export.SerialOffset += ExportOffsetDelta;
+			Writer << Export;
+		}
+	}
+
+	if (Writer.IsError())
+	{
+		UE_LOG(LogAssetHeaderPatcher, Error, TEXT("Failed to write to %s"), *DstAsset);
+		return EResult::ErrorFailedToWriteToDestinationFile;
+	}
+
+	if (AsetRegistryData.PkgData.DependencyDataOffset != INDEX_NONE)
+	{
+		// Re-write asset registry dependency data offset
+		Writer.Seek(Summary.AssetRegistryDataOffset);
+		Writer << AsetRegistryData.PkgData.DependencyDataOffset;
+
+		if (Writer.IsError())
+		{
+			UE_LOG(LogAssetHeaderPatcher, Error, TEXT("Failed to write to %s"), *DstAsset);
+			return EResult::ErrorFailedToWriteToDestinationFile;
+		}
+	}
+
+	return EResult::Success;
+}
diff --git a/Engine/Source/Developer/AssetTools/Private/AssetHeaderPatcher.h b/Engine/Source/Developer/AssetTools/Private/AssetHeaderPatcher.h
new file mode 100644
--- /dev/null
+++ b/Engine/Source/Developer/AssetTools/Private/AssetHeaderPatcher.h
@@ -0,0 +1,32 @@
+// Copyright Epic Games, Inc. All Rights Reserved.
+
+#pragma once
+
+#include "Tasks/Task.h"
+#include "Containers/ContainersFwd.h"
+
+struct FAssetHeaderPatcher
+{
+	enum class EResult
+	{
+		None,
+		Success,
+		ErrorFailedToLoadSourceAsset,
+		ErrorFailedToDeserializeSourceAsset,
+		ErrorUnexpectedSectionOrder,
+		ErrorBadOffset,
+		ErrorUnkownSection,
+		ErrorFailedToOpenDestinationFile,
+		ErrorFailedToWriteToDestinationFile,
+		ErrorEmptyRequireSection,
+	};
+
+	static UE::Tasks::TTask<EResult> Start(FString InSrcAsset, FString InDstAsset,
+		TMap<FString, FString> InSearchAndReplace);
+
+	static UE::Tasks::TTask<EResult> Start(TUniquePtr<FArchive> InSrcReader, FString InDstAsset,
+		TMap<FString, FString> InSearchAndReplace);
+
+	static EResult Test_DoPatch(FArchive& InSrcReader, FArchive& InDstWriter,
+		TMap<FString, FString> InSearchAndReplace);
+};
diff --git a/Engine/Source/Runtime/Engine/Public/WorldPartition/WorldPartitionActorDesc.h b/Engine/Source/Runtime/Engine/Public/WorldPartition/WorldPartitionActorDesc.h
--- a/Engine/Source/Runtime/Engine/Public/WorldPartition/WorldPartitionActorDesc.h
+++ b/Engine/Source/Runtime/Engine/Public/WorldPartition/WorldPartitionActorDesc.h
@@ -75,6 +75,7 @@ class FWorldPartitionActorDesc
 	friend struct FWorldPartitionActorDescUtils;
 	friend struct FWorldPartitionActorDescUnitTestAcccessor;
 	friend class FAssetRootPackagePatcher;
+	friend class FAssetHeaderPatcherInner;
 	friend class FActorDescArchive;
 	template<class U> friend class TActorDescContainerCollection;
 
