Fix ensure tripping due to letter case differences in drive letter
#rnx

diff --git a/Engine/Source/Developer/AssetTools/Private/AssetViewUtils.cpp b/Engine/Source/Developer/AssetTools/Private/AssetViewUtils.cpp
--- a/Engine/Source/Developer/AssetTools/Private/AssetViewUtils.cpp
+++ b/Engine/Source/Developer/AssetTools/Private/AssetViewUtils.cpp
@@ -1493,14 +1493,14 @@ int32 AssetViewUtils::GetPackageLengthForCooking(const FString& PackageName, boo
 	FString AbsoluteTargetPath = FPaths::ConvertRelativePathToFull(bIsEngineAsset ? FPaths::EngineDir() : FPaths::ProjectDir());
 
 	int32 AssetPathRelativeToCookRootLen = AbsolutePathToAsset.Len();
-	if (AbsolutePathToAsset.StartsWith(AbsoluteTargetPath, ESearchCase::CaseSensitive))
+	if (AbsolutePathToAsset.StartsWith(AbsoluteTargetPath, ESearchCase::IgnoreCase))
 	{
 		AssetPathRelativeToCookRootLen -= AbsoluteTargetPath.Len();
 	}
 	else if (ensureMsgf(PluginContainingAsset, TEXT("Only plugins can exist outside of the expected target path of '%s'. '%s' will not calculate an accurate result!"), *AbsoluteTargetPath, *AbsolutePathToAsset))
 	{
 		const FString AbsolutePluginRootPath = FPaths::ConvertRelativePathToFull(PluginContainingAsset->GetBaseDir());
-		if (ensureMsgf(AbsolutePathToAsset.StartsWith(AbsolutePluginRootPath, ESearchCase::CaseSensitive), TEXT("%s should start with %s"), *AbsolutePathToAsset, *AbsolutePluginRootPath))
+		if (ensureMsgf(AbsolutePathToAsset.StartsWith(AbsolutePluginRootPath, ESearchCase::IgnoreCase), TEXT("%s should start with %s"), *AbsolutePathToAsset, *AbsolutePluginRootPath))
 		{
 			AssetPathRelativeToCookRootLen -= AbsolutePluginRootPath.Len();
 			AssetPathRelativeToCookRootLen += FCString::Strlen(TEXT("Plugins/GameFeatures/XXXXXXXX-XXXX-XXXX-XXXX-XXXXXXXXXXXX")); // Use a GUID instead of the plugin name, as some external plugins cook as a GUID
@@ -1518,7 +1518,7 @@ int32 AssetViewUtils::GetPackageLengthForCooking(const FString& PackageName, boo
 	{
 		// We assume a constant size for the build machine base path for things that reside within the UE source tree
 		const FString AbsoluteUERootPath = FPaths::ConvertRelativePathToFull(FPaths::RootDir());
-		if (AbsoluteTargetPath.StartsWith(AbsoluteUERootPath, ESearchCase::CaseSensitive))
+		if (AbsoluteTargetPath.StartsWith(AbsoluteUERootPath, ESearchCase::IgnoreCase))
 		{
 			// Project is within the UE source tree, so remove the UE root path
 			int32 CookPathRelativeToTargetRootLen = CookSubPath.Len();
