Advanced copy using header patching as an option

Advanced copy, builds a list of depended assets, builds patching rules, and calls the AssetHeaderPatcher.
use `-ini:Engine:[ConsoleVariables]:AssetTools.UseHeaderPatchingAdvancedCopy=true` as cli option

diff --git a/Engine/Source/Developer/AssetTools/Private/AssetHeaderPatcher.cpp b/Engine/Source/Developer/AssetTools/Private/AssetHeaderPatcher.cpp
--- a/Engine/Source/Developer/AssetTools/Private/AssetHeaderPatcher.cpp
+++ b/Engine/Source/Developer/AssetTools/Private/AssetHeaderPatcher.cpp
@@ -67,6 +67,7 @@ namespace
 	{
 		Summary,
 		NameTable,
+		ImportTable,
 		ExportTable,
 		ThumbnailTable,
 		AssetRegistryData
@@ -141,20 +142,6 @@ namespace
 		TArray<FName>& NameTable;
 	};
 
-	int FindInPath(FStringView HayStack, FStringView Needle, int At = 0) 
-	{
-		check(At >= 0);
-		for (; At < HayStack.Len() - Needle.Len(); ++At)
-		{
-			if (FPathViews::Equals(FStringView(HayStack.GetData() + At, Needle.Len()), Needle))
-			{
-				return At;
-			}			
-		}
-
-		return -1;
-	}
-
 	struct FSummaryOffsetMeta
 	{
 		// NOTE: The offsets in Summary get to a max of 312 bytes.
@@ -246,6 +233,17 @@ namespace
 		return Out;
 	}
 
+	FStringView Find(const TMap<FString, FString>& Table, FStringView Needle)
+	{
+		uint32 NeedleHash = TMap<FString, FString>::KeyFuncsType::GetKeyHash<FStringView>(Needle);
+		const FString* MaybeNewItem = Table.FindByHash<FStringView>(NeedleHash, Needle);
+		if (MaybeNewItem)
+		{
+			return *MaybeNewItem;
+		}
+
+		return {};
+	}
 }
 
 // The information we need in the task to do patching.
@@ -395,7 +393,7 @@ FAssetHeaderPatcher::EResult FAssetHeaderPatcher::Test_DoPatch(FArchive& InReade
 FAssetHeaderPatcher::EResult FAssetHeaderPatcherInner::PatchHeader()
 {
 	FAssetHeaderPatcher::EResult Result = PatchHeader_Deserialize();
-	if (Result != EResult::None)
+	if (Result != EResult::Success)
 	{
 		return Result;
 	}
@@ -417,6 +415,10 @@ FAssetHeaderPatcher::EResult FAssetHeaderPatcherInner::PatchHeader_Deserialize()
 	MemAr.SetLicenseeUEVer(Summary.GetFileVersionLicenseeUE());
 	MemAr.SetEngineVer(Summary.SavedByEngineVersion);
 	MemAr.SetCustomVersions(Summary.GetCustomVersionContainer());
+	if (Summary.GetPackageFlags() & PKG_FilterEditorOnly)
+	{
+		MemAr.SetFilterEditorOnly(true);
+	}
 
 	if (Summary.DataResourceOffset > 0)
 	{
@@ -552,15 +554,53 @@ FAssetHeaderPatcher::EResult FAssetHeaderPatcherInner::PatchHeader_Deserialize()
 
 bool FAssetHeaderPatcherInner::DoPatch(FString& InOutString)
 {
-	for (const TPair<FString, FString>& KeyValue : SearchAndReplace)
 	{
-		if (FPathViews::Equals(InOutString, KeyValue.Key))
+		// Find a Path, change a Path.
+		FStringView MaybeReplacement = Find(SearchAndReplace, InOutString);
+		if (!MaybeReplacement.IsEmpty())
 		{
-			InOutString = KeyValue.Value;
+			InOutString = MaybeReplacement;
 			return true;
 		}
 	}
 
+	{
+		// Patch Object paths.
+		// Path occurs to the left of a ":"
+		int Idx{};
+		if (InOutString.FindChar(TCHAR(':'), Idx))
+		{
+			if (InOutString[Idx + 1] == TCHAR(':'))
+			{
+				// "::" is not a path delim
+				return false;
+			}
+			FStringView MaybeReplacement = Find(SearchAndReplace, FStringView(*InOutString, Idx));
+			if (!MaybeReplacement.IsEmpty())
+			{
+				FString Tmp = MaybeReplacement + InOutString.RightChop(Idx);;
+				InOutString = MoveTemp(Tmp);
+				return true;
+			}
+		}
+	}
+	
+	{
+		// Patch quoted paths.
+		// Path occurs to the right of the first "'" 
+		int Idx{};
+		if (InOutString.FindChar(TCHAR('\''), Idx) && InOutString.EndsWith(TEXT("'")))
+		{
+			FStringView MaybeReplacement = Find(SearchAndReplace, FStringView(*InOutString + Idx + 1, InOutString.Len() - (Idx + 2)));
+			if (!MaybeReplacement.IsEmpty())
+			{
+				FString Tmp = InOutString.LeftChop(Idx + 1) + MaybeReplacement + TCHAR('\'');
+				InOutString = MoveTemp(Tmp);
+				return true;
+			}
+		}
+	}
+
 	return false;
 }
 
@@ -592,11 +632,42 @@ void FAssetHeaderPatcherInner::PatchHeader_PatchSections()
 {
 	DoPatch(Summary.PackageName);
 
-	for (FName& Name : NameTable)
+	{	// Patch the Name table
+		// This data is used by all FGNames in the file.
+		// So if we patch a Identifier we append it to the name table.
+		// This is because there may be FNames in data we dont want to patch in structures we dont look at.
+		// If we patch the ident at inplace, then we would change those names.
+		TArray<FName> ToAppend;
+		for (FName& Name : NameTable)
+		{
+			FString TmpName = Name.GetPlainNameString();
+			if (DoPatch(TmpName)) 
+			{
+				// If the string does not contain any path seperators, then call it an Identifyer.
+				if (!(TmpName.Contains(TEXT("/")) || TmpName.Contains(TEXT("\\"))))
+				{
+					ToAppend.Add(FName(TmpName, NAME_NO_NUMBER));
+				} 
+				else
+				{
+					Name = FName(TmpName, Name.GetNumber());
+				}
+			}
+		}
+		NameTable.Append(ToAppend);
+
+		Summary.NameCount = NameTable.Num();
+	}
+
+	// Build Name Path table
+	// This is used to generate the indices saved by the FNamePatchingWriter
+	NameToIndexMap.Reserve(NameTable.Num());
+	for (TConstEnumerateRef<FName> Name : EnumerateRange(NameTable))
 	{
-		DoPatch(Name);
+		NameToIndexMap.Add(Name->GetDisplayIndex()) = Name.GetIndex();
 	}
 
+	// Import table
 	for (FObjectImport& Import : ImportTable)
 	{
 		DoPatch(Import.ObjectName);
@@ -671,13 +742,6 @@ void FAssetHeaderPatcherInner::PatchHeader_PatchSections()
 			}
 		}
 	}
-
-	// Build Name Path table
-	NameToIndexMap.Reserve(NameTable.Num());
-	for (TConstEnumerateRef<FName> Name : EnumerateRange(NameTable))
-	{
-		NameToIndexMap.Add(Name->GetDisplayIndex()) = Name.GetIndex();
-	}
 }
 
 FAssetHeaderPatcher::EResult FAssetHeaderPatcherInner::PatchHeader_WriteDestinationFile()
@@ -688,6 +752,7 @@ FAssetHeaderPatcher::EResult FAssetHeaderPatcherInner::PatchHeader_WriteDestinat
 	const FSectionData SourceSections[] = {
 		{ EPatchedSection::Summary,					0,										HeaderInformation.SummarySize,			true },
 		{ EPatchedSection::NameTable,				Summary.NameOffset,						HeaderInformation.NameTableSize,		true },
+		{ EPatchedSection::ImportTable,				Summary.ImportOffset,					HeaderInformation.ImportTableSize,		true },
 		{ EPatchedSection::ExportTable,				Summary.ExportOffset,					HeaderInformation.ExportTableSize,		true },
 		{ EPatchedSection::ThumbnailTable,			Summary.ThumbnailTableOffset,			HeaderInformation.ThumbnailTableSize,	false },
 		{ EPatchedSection::AssetRegistryData,		Summary.AssetRegistryDataOffset,		AsetRegistryData.SectionSize,			true },
@@ -740,6 +805,10 @@ FAssetHeaderPatcher::EResult FAssetHeaderPatcherInner::PatchHeader_WriteDestinat
 	Writer.SetLicenseeUEVer(Summary.GetFileVersionLicenseeUE());
 	Writer.SetEngineVer(Summary.SavedByEngineVersion);
 	Writer.SetCustomVersions(Summary.GetCustomVersionContainer());
+	if (Summary.GetPackageFlags() & PKG_FilterEditorOnly)
+	{
+		Writer.SetFilterEditorOnly(true);
+	}
 
 	int64 LastSectionEndedAt = 0;
 	
@@ -800,6 +869,23 @@ FAssetHeaderPatcher::EResult FAssetHeaderPatcherInner::PatchHeader_WriteDestinat
 				break;
 			}
 
+			case EPatchedSection::ImportTable:
+			{
+				const int64 ImportTableStartOffset = Writer.Tell();
+				for (FObjectImport& Import : ImportTable)
+				{
+					Writer << Import;
+				}
+				const int64 ImportTableSize = Writer.Tell() - ImportTableStartOffset;
+				const int64 Delta = ImportTableSize - SourceSection.Size;
+				check(Delta == 0);
+				checkf(ImportTableSize == SourceSection.Size, TEXT("%d == %d"), (int)ImportTableSize, (int)SourceSection.Size); // We only patch export table offsets, we should not be patching size
+				checkf(Summary.ImportCount == ImportTable.Num(), TEXT("%d == %d"), Summary.ImportCount, ImportTable.Num());
+				checkf(Summary.ImportOffset == ImportTableStartOffset, TEXT("%d == %d"), Summary.ImportOffset, ImportTableStartOffset);
+
+				break;
+			}
+
 			case EPatchedSection::ExportTable:
 			{
 				// The export table offsets aren't correct yet.
@@ -811,8 +897,7 @@ FAssetHeaderPatcher::EResult FAssetHeaderPatcherInner::PatchHeader_WriteDestinat
 				}
 				const int64 ExportTableSize = Writer.Tell() - ExportTableStartOffset;
 				const int64 Delta = ExportTableSize - SourceSection.Size;
-				PatchSummaryOffsets(Summary, ExportTableStartOffset, Delta);
-
+				check(Delta == 0);
 				checkf(ExportTableSize == SourceSection.Size, TEXT("%d == %d"), (int)ExportTableSize, (int)SourceSection.Size); // We only patch export table offsets, we should not be patching size
 				checkf(Summary.ExportCount == ExportTable.Num(), TEXT("%d == %d"), Summary.ExportCount, ExportTable.Num());
 				checkf(Summary.ExportOffset == ExportTableStartOffset, TEXT("%d == %d"), Summary.ExportOffset, ExportTableStartOffset);
diff --git a/Engine/Source/Developer/AssetTools/Private/AssetHeaderPatcher.h b/Engine/Source/Developer/AssetTools/Private/AssetHeaderPatcher.h
--- a/Engine/Source/Developer/AssetTools/Private/AssetHeaderPatcher.h
+++ b/Engine/Source/Developer/AssetTools/Private/AssetHeaderPatcher.h
@@ -9,7 +9,6 @@ struct FAssetHeaderPatcher
 {
 	enum class EResult
 	{
-		None,
 		Success,
 		ErrorFailedToLoadSourceAsset,
 		ErrorFailedToDeserializeSourceAsset,
diff --git a/Engine/Source/Developer/AssetTools/Private/AssetTools.cpp b/Engine/Source/Developer/AssetTools/Private/AssetTools.cpp
--- a/Engine/Source/Developer/AssetTools/Private/AssetTools.cpp
+++ b/Engine/Source/Developer/AssetTools/Private/AssetTools.cpp
@@ -110,6 +110,9 @@
 #include "Factories/SceneImportFactory.h"
 #include "Misc/AssetFilterData.h"
 
+#include "AssetHeaderPatcher.h"
+#include "Algo/Copy.h"
+
 #include UE_INLINE_GENERATED_CPP_BY_NAME(AssetTools)
 
 #define LOCTEXT_NAMESPACE "AssetTools"
@@ -443,6 +446,18 @@ namespace UE::AssetTools::Private
 				UE_LOG(LogAssetTools, Log, TEXT("Folder Write Permissions:\n%s"), *IAssetTools::Get().GetWritableFolderPermissionList()->ToString());
 			}));
 
+	/** 
+	 * CVar to specify if we should use Header patching in advanced copy.
+	 * Default is false.
+	 */
+	bool bEnableHeaderPatching = false;
+	FAutoConsoleVariableRef CVarEnableHeaderPatching(
+		TEXT("AssetTools.UseHeaderPatchingAdvancedCopy"),
+		bEnableHeaderPatching,
+		TEXT("If set to true, this will use Header Patching to copy the files instead of performing a full load."),
+		ECVF_Default
+	);
+
 	// use a struct as a namespace to allow easier friend declarations
 	struct FPackageMigrationImpl
 	{
@@ -2135,6 +2150,88 @@ void UAssetToolsImpl::GetAllAdvancedCopySources(FName SelectedPackage, FAdvanced
 	}
 }
 
+namespace 
+{
+bool IsSlashOrBackslash(TCHAR C) 
+{
+	return C == TEXT('/') || C == TEXT('\\'); 
+}
+	
+TMap<FString, FString> AllSourceAndDestPackages(const TMap<FString, FString>& SourceAndDestPackages)
+{
+	TMap<FString, FString> Result;
+
+	IAssetRegistry& Registry = *IAssetRegistry::Get();
+
+	TArray< TTuple<FString, FString> > ToProcess;
+	Algo::Copy(SourceAndDestPackages, ToProcess);
+
+	while (ToProcess.Num())
+	{
+		TTuple<FString, FString> Package = ToProcess.Pop();
+
+		if (Result.Contains(Package.Key))
+		{
+			continue;
+		}
+
+		// Become a patching name even if it doesn't have a file.
+		Result.Add({ Package.Key, Package.Value });
+
+		TArray<FName> Dependencies;
+
+		if (!Registry.GetDependencies(FName(*Package.Key), Dependencies))
+		{
+			continue;
+		}
+
+		// Making String Views into strings because the String.Replace used inside the loop cannot use the views.
+		FString SrcPackageRoot = FString(FPackageName::SplitPackageNameRoot(Package.Key, nullptr));
+		FString DstPackageRoot = FString(FPackageName::SplitPackageNameRoot(Package.Value, nullptr));
+
+		for (const FName Dependency : Dependencies)
+		{
+			const FString SrcDependencyString = Dependency.ToString();
+			
+			// checking from +1 Dependency has a leading '/' 
+			if (IsSlashOrBackslash(SrcDependencyString[0])
+				&& FStringView(*SrcDependencyString + 1, SrcPackageRoot.Len()) == SrcPackageRoot
+				&& IsSlashOrBackslash(SrcDependencyString[SrcPackageRoot.Len() + 1]))
+			{
+				// if a dep start with the package name, then we are going to copy the asset.
+				// but we need to recurse on this asset as it may have sub dependencies we don't know of yet.
+				const FString DstDependencyString = SrcDependencyString.Replace(*SrcPackageRoot, *DstPackageRoot, ESearchCase::CaseSensitive);
+				ToProcess.Add({ SrcDependencyString , DstDependencyString });
+			}
+		}
+	}
+
+	return Result;
+}
+
+TMap<FString, FString> GenerateAdditionalAssetMappings(const TMap<FString, FString>& SourceAndDestPackages)
+{
+	TMap<FString, FString> Result;
+
+	for (const TTuple<FString, FString>& Package : SourceAndDestPackages)
+	{
+		// FPathViews::GetBaseFilename gives the same result as FPackageName::GetShortName
+		// for a file path, but returns a StringView not a String.
+		FStringView SrcPackageName = FPathViews::GetBaseFilename(Package.Key); 
+		FStringView DstPackageName = FPathViews::GetBaseFilename(Package.Value);
+
+		// Inject Path.ObjectName
+		Result.Add({ Package.Key + TCHAR('.') + SrcPackageName, Package.Value + TCHAR('.') + DstPackageName });
+		if (SrcPackageName != DstPackageName)
+		{
+			Result.Add({ FString(SrcPackageName), FString(DstPackageName) });
+		}
+	}
+
+	return Result;
+}
+}
+
 bool UAssetToolsImpl::AdvancedCopyPackages(
 	const TMap<FString, FString>& SourceAndDestPackages,
 	const bool bForceAutosave,
@@ -2142,6 +2239,8 @@ bool UAssetToolsImpl::AdvancedCopyPackages(
 	FDuplicatedObjects* OutDuplicatedObjects,
 	EMessageSeverity::Type NotificationSeverityFilter) const
 {
+	TRACE_CPUPROFILER_EVENT_SCOPE(AdvancedCopyPackages);
+
 	if (ValidateFlattenedAdvancedCopyDestinations(SourceAndDestPackages))
 	{
 		TArray<FString> SuccessfullyCopiedDestinationFiles;
@@ -2163,58 +2262,110 @@ bool UAssetToolsImpl::AdvancedCopyPackages(
 		ExistingObjectSet.Reserve(SourceAndDestPackages.Num());
 		NewObjectSet.Reserve(SourceAndDestPackages.Num());
 
-		TUniquePtr<FScopedSlowTask> LoopProgress = MakeUnique<FScopedSlowTask>(static_cast<float>(SourceAndDestPackages.Num()), LOCTEXT("AdvancedCopyPackages.CopyingFilesAndDependencies", "Copying Files and Dependencies..."));
-		LoopProgress->MakeDialog();
+		TUniquePtr<FScopedSlowTask> LoopProgress;
 
-		for (const auto& Package : SourceAndDestPackages)
+		if (UE::AssetTools::Private::bEnableHeaderPatching)
 		{
-			const FString& PackageName = Package.Key;
-			const FString& DestFilename = Package.Value;
-			FString SrcFilename;
+			std::atomic<int32> PatchAssetsCompletedCount = 0;
+			UE::Tasks::FTaskEvent PatchAssetsCompletionTask{ UE_SOURCE_LOCATION };
 
-			if (FPackageName::DoesPackageExist(PackageName, &SrcFilename))
+			TMap<FString, FString> ToCopyAndPatch = AllSourceAndDestPackages(SourceAndDestPackages);
+			TMap<FString, FString> PatchingPatterns = GenerateAdditionalAssetMappings(SourceAndDestPackages);
+			PatchingPatterns.Append(ToCopyAndPatch);
+
+			LoopProgress = MakeUnique<FScopedSlowTask>(static_cast<float>(ToCopyAndPatch.Num()), LOCTEXT("AdvancedCopyPackages.ReplacingAssetReferences", "Replacing Asset References..."));
+			LoopProgress->MakeDialog();
+
+			for (const TTuple<FString, FString>& Package : ToCopyAndPatch)
 			{
 				LoopProgress->EnterProgressFrame();
-				UPackage* Pkg = LoadPackage(nullptr, *PackageName, LOAD_None);
-				if (Pkg)
+
+				const FString& PackageName = Package.Key;
+				const FString& DestPackage = Package.Value;
+				FString SrcFilename;
+
+				if (FPackageName::DoesPackageExist(PackageName, &SrcFilename))
+				{
+					FString DestFilename = FPackageName::LongPackageNameToFilename(DestPackage, FString(FPathViews::GetExtension(SrcFilename, true)));
+					UE::Tasks::TTask<FAssetHeaderPatcher::EResult> PatcherTask = FAssetHeaderPatcher::Start(SrcFilename, DestFilename, PatchingPatterns);
+					UE::Tasks::FTask PatcherLifetimeTask = UE::Tasks::Launch(UE_SOURCE_LOCATION, [&PatchAssetsCompletedCount]
+						{
+							PatchAssetsCompletedCount.fetch_add(1, std::memory_order_relaxed);
+						},
+						MoveTemp(PatcherTask), UE::Tasks::ETaskPriority::Default, UE::Tasks::EExtendedTaskPriority::Inline);
+					PatchAssetsCompletionTask.AddPrerequisites(PatcherLifetimeTask);
+				}
+			}
+
+			PatchAssetsCompletionTask.Trigger();
+			while (!PatchAssetsCompletionTask.Wait(FTimespan::FromSeconds(0.5)))
+			{
+				LoopProgress->CompletedWork = (float)PatchAssetsCompletedCount.load(std::memory_order_relaxed);
+				LoopProgress->TickProgress();
+			}
+
+			for (const TTuple<FString, FString>& Package : SourceAndDestPackages)
+			{			
+				SuccessfullyCopiedSourcePackages.Add(FName(Package.Key));
+				SuccessfullyCopiedDestinationFiles.Add(Package.Value);
+			}
+		}
+		else
+		{
+			LoopProgress = MakeUnique<FScopedSlowTask>(static_cast<float>(SourceAndDestPackages.Num()), LOCTEXT("AdvancedCopyPackages.CopyingFilesAndDependencies", "Copying Files and Dependencies..."));
+			LoopProgress->MakeDialog();
+
+			for (const auto& Package : SourceAndDestPackages)
+			{
+				const FString& PackageName = Package.Key;
+				const FString& DestPackage = Package.Value;
+				FString SrcFilename;
+
+				if (FPackageName::DoesPackageExist(PackageName, &SrcFilename))
 				{
-					FString Name = ObjectTools::SanitizeObjectName(FPaths::GetBaseFilename(SrcFilename));
-					UObject* ExistingObject = StaticFindObject(UObject::StaticClass(), Pkg, *Name);
-					if (ExistingObject)
+					LoopProgress->EnterProgressFrame();
+
+					UPackage* Pkg = LoadPackage(nullptr, *PackageName, LOAD_None);
+					if (Pkg)
 					{
-						TSet<UPackage*> ObjectsUserRefusedToFullyLoad;
-						ObjectTools::FPackageGroupName PGN;
-						PGN.GroupName = TEXT("");
-						PGN.ObjectName = FPaths::GetBaseFilename(DestFilename);
-						PGN.PackageName = DestFilename;
-						const bool bShouldPromptForDestinationConflict = !bCopyOverAllDestinationOverlaps;
-						TMap<TSoftObjectPtr<UObject>, TSoftObjectPtr<UObject>> DuplicatedObjects;
-						FName PackageFName(*PackageName);
-
-
-						// Temp fix for some codepaths that allows advanced copy of world packages. For partitioned worlds, this can only be supported for worlds with
-						// streaming disabled and this code should be removed once the callers switch to the same codepath as editor save as.
-						if (UWorld* World = Cast<UWorld>(ExistingObject))
+						FString Name = ObjectTools::SanitizeObjectName(FPaths::GetBaseFilename(SrcFilename));
+						UObject* ExistingObject = StaticFindObject(UObject::StaticClass(), Pkg, *Name);
+						if (ExistingObject)
 						{
-							if (UWorldPartition* WorldPartition = World->GetWorldPartition())
+							TSet<UPackage*> ObjectsUserRefusedToFullyLoad;
+							ObjectTools::FPackageGroupName PGN;
+							PGN.GroupName = TEXT("");
+							PGN.ObjectName = FPaths::GetBaseFilename(DestPackage);
+							PGN.PackageName = DestPackage;
+							const bool bShouldPromptForDestinationConflict = !bCopyOverAllDestinationOverlaps;
+							TMap<TSoftObjectPtr<UObject>, TSoftObjectPtr<UObject>> DuplicatedObjects;
+							FName PackageFName(*PackageName);
+
+
+							// Temp fix for some codepaths that allows advanced copy of world packages. For partitioned worlds, this can only be supported for worlds with
+							// streaming disabled and this code should be removed once the callers switch to the same codepath as editor save as.
+							if (UWorld* World = Cast<UWorld>(ExistingObject))
 							{
-								check(!WorldPartition->IsStreamingEnabled());
-								if (!WorldPartition->IsInitialized())
+								if (UWorldPartition* WorldPartition = World->GetWorldPartition())
 								{
-									WorldPartition->Initialize(World, FTransform::Identity);
+									check(!WorldPartition->IsStreamingEnabled());
+									if (!WorldPartition->IsInitialized())
+									{
+										WorldPartition->Initialize(World, FTransform::Identity);
+									}
+									CopiedWorldPartitionMaps.Add(PackageFName);
 								}
-								CopiedWorldPartitionMaps.Add(PackageFName);
 							}
-						}
 
-						if (UObject* NewObject = ObjectTools::DuplicateSingleObject(ExistingObject, PGN, ObjectsUserRefusedToFullyLoad, bShouldPromptForDestinationConflict, &DuplicatedObjects))
-						{
-							ExistingObjectSet.Add(ExistingObject);
-							NewObjectSet.Add(NewObject);
-							DuplicatedObjectsForEachPackage.Add(MoveTemp(DuplicatedObjects));
-							SuccessfullyCopiedSourcePackages.Add(PackageFName);
-							SuccessfullyCopiedDestinationFiles.Add(DestFilename);
-							SuccessfullyCopiedDestinationPackages.Add(NewObject->GetPackage());
+							if (UObject* NewObject = ObjectTools::DuplicateSingleObject(ExistingObject, PGN, ObjectsUserRefusedToFullyLoad, bShouldPromptForDestinationConflict, &DuplicatedObjects))
+							{
+								ExistingObjectSet.Add(ExistingObject);
+								NewObjectSet.Add(NewObject);
+								DuplicatedObjectsForEachPackage.Add(MoveTemp(DuplicatedObjects));
+								SuccessfullyCopiedSourcePackages.Add(PackageFName);
+								SuccessfullyCopiedDestinationFiles.Add(DestPackage);
+								SuccessfullyCopiedDestinationPackages.Add(NewObject->GetPackage());
+							}
 						}
 					}
 				}
@@ -2272,19 +2423,22 @@ bool UAssetToolsImpl::AdvancedCopyPackages(
 				}
 			}
 
-			for (FName Dependency : Dependencies)
+			if (!UE::AssetTools::Private::bEnableHeaderPatching)
 			{
-				const int32 DependencyIndex = SuccessfullyCopiedSourcePackages.IndexOfByKey(Dependency);
-				if (DependencyIndex != INDEX_NONE)
+				for (FName Dependency : Dependencies)
 				{
-					Consolidations.Reserve(Consolidations.Num() + DuplicatedObjectsForEachPackage[DependencyIndex].Num());
-					for (const TPair<TSoftObjectPtr<UObject>, TSoftObjectPtr<UObject>>& Duplication : DuplicatedObjectsForEachPackage[DependencyIndex])
+					const int32 DependencyIndex = SuccessfullyCopiedSourcePackages.IndexOfByKey(Dependency);
+					if (DependencyIndex != INDEX_NONE)
 					{
-						UObject* SourceObject = Duplication.Key.Get();
-						UObject* NewObject = Duplication.Value.Get();
-						if (SourceObject && NewObject)
+						Consolidations.Reserve(Consolidations.Num() + DuplicatedObjectsForEachPackage[DependencyIndex].Num());
+						for (const TPair<TSoftObjectPtr<UObject>, TSoftObjectPtr<UObject>>& Duplication : DuplicatedObjectsForEachPackage[DependencyIndex])
 						{
-							Consolidations.FindOrAdd(NewObject).AddUnique(SourceObject);
+							UObject* SourceObject = Duplication.Key.Get();
+							UObject* NewObject = Duplication.Value.Get();
+							if (SourceObject && NewObject)
+							{
+								Consolidations.FindOrAdd(NewObject).AddUnique(SourceObject);
+							}
 						}
 					}
 				}
