Patching softpaths

diff --git a/Engine/Source/Developer/AssetTools/Private/AssetHeaderPatcher.cpp b/Engine/Source/Developer/AssetTools/Private/AssetHeaderPatcher.cpp
--- a/Engine/Source/Developer/AssetTools/Private/AssetHeaderPatcher.cpp
+++ b/Engine/Source/Developer/AssetTools/Private/AssetHeaderPatcher.cpp
@@ -67,8 +67,10 @@ namespace
 	{
 		Summary,
 		NameTable,
+		SoftPathTable,
 		ImportTable,
 		ExportTable,
+		SoftPackageReferencesTable,
 		ThumbnailTable,
 		AssetRegistryData
 	};
@@ -299,6 +301,8 @@ public:
 	FHeaderInformation HeaderInformation;
 	FPackageFileSummary Summary;
 	TArray<FName> NameTable;
+	TArray<FSoftObjectPath> SoftObjectPathTable;
+	TArray<FName> SoftPackageReferencesTable;
 	TMap<FNameEntryId, int32> NameToIndexMap;
 	TArray<FObjectImport> ImportTable;
 	TArray<FObjectExport> ExportTable;
@@ -441,6 +445,22 @@ FAssetHeaderPatcher::EResult FAssetHeaderPatcherInner::PatchHeader_Deserialize()
 		HeaderInformation.NameTableSize = MemAr.Tell() - HeaderInformation.SummarySize;
 	}
 
+	if (Summary.SoftObjectPathsCount > 0)
+	{
+		MemAr.Seek(Summary.SoftObjectPathsOffset);
+		SoftObjectPathTable.Reserve(Summary.SoftObjectPathsCount);
+		for (int32 Idx = 0; Idx < Summary.SoftObjectPathsCount; ++Idx)
+		{
+			FSoftObjectPath& PathRef = SoftObjectPathTable.AddDefaulted_GetRef();
+			PathRef.SerializePath(MemAr);
+		}
+		HeaderInformation.SoftObjectPathListSize = MemAr.Tell() - Summary.SoftObjectPathsOffset;
+	}
+	else
+	{
+		HeaderInformation.SoftObjectPathListSize = 0;
+	}
+
 #define UE_CHECK_AND_SET_ERROR_AND_RETURN(EXP)	\
 	do											\
 	{											\
@@ -494,6 +514,23 @@ FAssetHeaderPatcher::EResult FAssetHeaderPatcherInner::PatchHeader_Deserialize()
 
 #undef UE_CHECK_AND_SET_ERROR_AND_RETURN
 
+	if (Summary.SoftPackageReferencesCount)
+	{
+		MemAr.Seek(Summary.SoftPackageReferencesOffset);
+		SoftPackageReferencesTable.Reserve(Summary.SoftPackageReferencesCount);
+		for (int32 Idx = 0; Idx < Summary.SoftPackageReferencesCount; ++Idx)
+		{
+			FName& Reference = SoftPackageReferencesTable.Emplace_GetRef();
+			MemAr << Reference;
+		}
+
+		HeaderInformation.SoftPackageReferencesListSize = MemAr.Tell() - Summary.SoftPackageReferencesOffset;
+	}
+	else
+	{
+		HeaderInformation.SoftPackageReferencesListSize = 0;
+	}
+
 	if (Summary.ThumbnailTableOffset)
 	{
 		MemAr.Seek(Summary.ThumbnailTableOffset);
@@ -640,17 +677,19 @@ void FAssetHeaderPatcherInner::PatchHeader_PatchSections()
 		TArray<FName> ToAppend;
 		for (FName& Name : NameTable)
 		{
-			FString TmpName = Name.GetPlainNameString();
+			FName TmpName = Name;
+
 			if (DoPatch(TmpName)) 
 			{
-				// If the string does not contain any path seperators, then call it an Identifyer.
-				if (!(TmpName.Contains(TEXT("/")) || TmpName.Contains(TEXT("\\"))))
+				FString TmpNameStr = Name.GetPlainNameString();
+				// If the string does not contain any path separators, then call it an Identifier.
+				if (!(TmpNameStr.Contains(TEXT("/")) || TmpNameStr.Contains(TEXT("\\"))))
 				{
-					ToAppend.Add(FName(TmpName, NAME_NO_NUMBER));
+					ToAppend.Add(TmpName);
 				} 
 				else
 				{
-					Name = FName(TmpName, Name.GetNumber());
+					Name = TmpName;
 				}
 			}
 		}
@@ -667,6 +706,21 @@ void FAssetHeaderPatcherInner::PatchHeader_PatchSections()
 		NameToIndexMap.Add(Name->GetDisplayIndex()) = Name.GetIndex();
 	}
 
+	// Soft paths
+	for (FSoftObjectPath& PathRef : SoftObjectPathTable)
+	{
+		FTopLevelAssetPath AssetPath = PathRef.GetAssetPath();
+		FName PackageName = AssetPath.GetPackageName();
+		FName AssetName = AssetPath.GetAssetName();
+		bool bPatched = DoPatch(PackageName);
+		bPatched |= DoPatch(AssetName);
+		if (bPatched)
+		{
+			FTopLevelAssetPath NewAssetPath{ PackageName, AssetName };
+			PathRef.SetPath(NewAssetPath, PathRef.GetSubPathString());
+		}
+	}
+
 	// Import table
 	for (FObjectImport& Import : ImportTable)
 	{
@@ -682,6 +736,7 @@ void FAssetHeaderPatcherInner::PatchHeader_PatchSections()
 #endif
 	}
 
+	// Export Table
 	for (FObjectExport& Export : ExportTable)
 	{
 		DoPatch(Export.ObjectName);
@@ -690,6 +745,12 @@ void FAssetHeaderPatcherInner::PatchHeader_PatchSections()
 #endif
 	}
 
+	// Soft Package Reference's
+	for (FName& Reference : SoftPackageReferencesTable)
+	{
+		DoPatch(Reference);
+	}
+
 	// Asset Register Data
 	for (FAsetRegistryObjectData& ObjData : AsetRegistryData.ObjectData)
 	{
@@ -748,14 +809,16 @@ FAssetHeaderPatcher::EResult FAssetHeaderPatcherInner::PatchHeader_WriteDestinat
 {
 	// Serialize modified sections and reconstruct the file	
 	// Original offsets and sizes of any sections that will be patched
-	//	Tag											Offset									Size						bRequired
+	//	Tag												Offset									Size								bRequired
 	const FSectionData SourceSections[] = {
-		{ EPatchedSection::Summary,					0,										HeaderInformation.SummarySize,			true },
-		{ EPatchedSection::NameTable,				Summary.NameOffset,						HeaderInformation.NameTableSize,		true },
-		{ EPatchedSection::ImportTable,				Summary.ImportOffset,					HeaderInformation.ImportTableSize,		true },
-		{ EPatchedSection::ExportTable,				Summary.ExportOffset,					HeaderInformation.ExportTableSize,		true },
-		{ EPatchedSection::ThumbnailTable,			Summary.ThumbnailTableOffset,			HeaderInformation.ThumbnailTableSize,	false },
-		{ EPatchedSection::AssetRegistryData,		Summary.AssetRegistryDataOffset,		AsetRegistryData.SectionSize,			true },
+		{ EPatchedSection::Summary,						0,										HeaderInformation.SummarySize,			true },
+		{ EPatchedSection::NameTable,					Summary.NameOffset,						HeaderInformation.NameTableSize,		true },
+		{ EPatchedSection::SoftPathTable,				Summary.SoftObjectPathsOffset,			HeaderInformation.SoftObjectPathListSize, false },
+		{ EPatchedSection::ImportTable,					Summary.ImportOffset,					HeaderInformation.ImportTableSize,		true },
+		{ EPatchedSection::ExportTable,					Summary.ExportOffset,					HeaderInformation.ExportTableSize,		true },
+		{ EPatchedSection::SoftPackageReferencesTable,	Summary.SoftPackageReferencesOffset,	HeaderInformation.SoftPackageReferencesListSize, false },
+		{ EPatchedSection::ThumbnailTable,				Summary.ThumbnailTableOffset,			HeaderInformation.ThumbnailTableSize,	false },
+		{ EPatchedSection::AssetRegistryData,			Summary.AssetRegistryDataOffset,		AsetRegistryData.SectionSize,			true },
 	};
 
 	const int32 SourceTotalHeaderSize = Summary.TotalHeaderSize;
@@ -869,6 +932,22 @@ FAssetHeaderPatcher::EResult FAssetHeaderPatcherInner::PatchHeader_WriteDestinat
 				break;
 			}
 
+			case EPatchedSection::SoftPathTable:
+			{
+				const int64 TableStartOffset = Writer.Tell();
+				for (FSoftObjectPath& PathRef : SoftObjectPathTable)
+				{
+					PathRef.SerializePath(Writer);
+				}
+				const int64 TableSize = Writer.Tell() - TableStartOffset;
+				const int64 Delta = TableSize - SourceSection.Size;
+				checkf(Delta == 0, TEXT("Delta should be Zero. is %d"), (int)Delta);
+				check(Summary.SoftObjectPathsCount == SoftObjectPathTable.Num());
+				check(Summary.SoftObjectPathsOffset == TableStartOffset);
+
+				break;
+			}
+
 			case EPatchedSection::ImportTable:
 			{
 				const int64 ImportTableStartOffset = Writer.Tell();
@@ -905,6 +984,22 @@ FAssetHeaderPatcher::EResult FAssetHeaderPatcherInner::PatchHeader_WriteDestinat
 				break;
 			}
 
+			case EPatchedSection::SoftPackageReferencesTable:
+			{
+				const int64 TableStartOffset = Writer.Tell();
+				for (FName& Reference : SoftPackageReferencesTable)
+				{
+					Writer << Reference;
+				}
+				const int64 TableSize = Writer.Tell() - TableStartOffset;
+				const int64 Delta = TableSize - SourceSection.Size;
+				checkf(Delta == 0, TEXT("Delta should be Zero. is %d"), (int)Delta);
+				check(Summary.SoftPackageReferencesCount == SoftPackageReferencesTable.Num());
+				check(Summary.SoftPackageReferencesOffset == TableStartOffset);
+
+				break;
+			}
+
 			case EPatchedSection::ThumbnailTable:
 			{
 				// Luckily the strings in the thumbnail table don't include the package, so only the offsets need to be patched
