Filtering package names through FName to make sain patching set for header patching

diff --git a/Engine/Source/Developer/AssetTools/Private/AssetTools.cpp b/Engine/Source/Developer/AssetTools/Private/AssetTools.cpp
--- a/Engine/Source/Developer/AssetTools/Private/AssetTools.cpp
+++ b/Engine/Source/Developer/AssetTools/Private/AssetTools.cpp
@@ -2213,15 +2213,35 @@ TMap<FString, FString> GenerateAdditionalAssetMappings(const TMap<FString, FStri
 {
 	TMap<FString, FString> Result;
 
+	TCHAR SrcNameBuffer[NAME_SIZE];
+	TCHAR DstNameBuffer[NAME_SIZE];
 	for (const TTuple<FString, FString>& Package : SourceAndDestPackages)
 	{
+		// We make FName's out of our incoming string package names
+		// as on rare occasion some of them have a '_[0-9]+' tail.
+		// Making FNames parse and strip this number on construction
+		// which then makes it consistent with how the are found in the name and import tables.
+		FName SrcName = *Package.Key;
+		int32 SrcNameLen = (int32)SrcName.GetPlainNameString(SrcNameBuffer);
+		FStringView SrcNameView{ SrcNameBuffer, SrcNameLen };
+
+		FName DstName = *Package.Value;
+		int32 DstNameLen = (int32)DstName.GetPlainNameString(DstNameBuffer);
+		FStringView DstNameView{ DstNameBuffer, DstNameLen };		
+
+		if (SrcNameLen != Package.Key.Len())
+		{
+			Result.Add({ FString(SrcNameView), FString(DstNameView) });
+		}
+
 		// FPathViews::GetBaseFilename gives the same result as FPackageName::GetShortName
 		// for a file path, but returns a StringView not a String.
-		FStringView SrcPackageName = FPathViews::GetBaseFilename(Package.Key); 
-		FStringView DstPackageName = FPathViews::GetBaseFilename(Package.Value);
+		FStringView SrcPackageName = FPathViews::GetBaseFilename(SrcNameView);
+		FStringView DstPackageName = FPathViews::GetBaseFilename(DstNameView);
 
 		// Inject Path.ObjectName
-		Result.Add({ Package.Key + TCHAR('.') + SrcPackageName, Package.Value + TCHAR('.') + DstPackageName });
+		// NOTE: this would be better to use a string builder.
+		Result.Add({ FString(SrcNameView) + TCHAR('.') + SrcPackageName, FString(DstNameView) + TCHAR('.') + DstPackageName });
 		if (SrcPackageName != DstPackageName)
 		{
 			Result.Add({ FString(SrcPackageName), FString(DstPackageName) });
