* Create 'Check in before syncing older snapshot' dialog
* Create Warning dialog for further rewind
* Implement 'Sync snapshot' functionality
* Implement 'Rewind to' functionality

diff --git a/Engine/Source/Developer/AssetTools/Private/AssetViewUtils.cpp b/Engine/Source/Developer/AssetTools/Private/AssetViewUtils.cpp
--- a/Engine/Source/Developer/AssetTools/Private/AssetViewUtils.cpp
+++ b/Engine/Source/Developer/AssetTools/Private/AssetViewUtils.cpp
@@ -1846,7 +1846,7 @@ bool AssetViewUtils::SyncPackagesFromSourceControl(const TArray<FString>& Packag
 	return true;
 }
 
-static bool SyncPathsFromSourceControl(const TArray<FString>& Paths, bool bCheckDependencies)
+static bool SyncPathsFromSourceControl(const FString& Revision, const TArray<FString>& Paths, bool bCheckDependencies)
 {
 	TArray<FString> PathsOnDisk;
 	PathsOnDisk.Reserve(Paths.Num());
@@ -1882,6 +1882,7 @@ static bool SyncPathsFromSourceControl(const TArray<FString>& Paths, bool bCheck
 
 		// Use FSyncPreview if possible.
 		TSharedRef<FSyncPreview> PreviewOperation = ISourceControlOperation::Create<FSyncPreview>();
+		PreviewOperation->SetRevision(Revision);
 		if (SCCProvider.CanExecuteOperation(PreviewOperation) &&
 			SCCProvider.Execute(PreviewOperation, PathsOnDisk) == ECommandResult::Succeeded)
 		{
@@ -1993,7 +1994,9 @@ static bool SyncPathsFromSourceControl(const TArray<FString>& Paths, bool bCheck
 		PathsToSync.Append(SourceControlHelpers::PackageFilenames(ExtraPackagesToSync));
 
 		// Sync everything...
-		ECommandResult::Type SyncResult = SCCProvider.Execute(ISourceControlOperation::Create<FSync>(), PathsToSync);
+		TSharedRef<FSync> Operation = ISourceControlOperation::Create<FSync>();
+		Operation->SetRevision(Revision);
+		ECommandResult::Type SyncResult = SCCProvider.Execute(Operation, PathsToSync);
 
 		// Syncing may have deleted some packages, so we need to unload those rather than re-load them...
 		// Note: we will store the package using weak pointers here otherwise we might have garbage collection issues after the ReloadPackages call
@@ -2056,9 +2059,9 @@ static bool SyncPathsFromSourceControl(const TArray<FString>& Paths, bool bCheck
 	return true;
 }
 
-bool AssetViewUtils::SyncLatestFromSourceControl()
+static bool SyncPathsFromSourceControl(const TArray<FString>& Paths, bool bCheckDependencies)
 {
-	return SyncPathsFromSourceControl(SourceControlHelpers::GetSourceControlLocations(), /*bCheckDependencies=*/false);
+	return SyncPathsFromSourceControl(TEXT(""), Paths, bCheckDependencies);
 }
 
 bool AssetViewUtils::SyncPathsFromSourceControl(const TArray<FString>& Paths)
@@ -2066,6 +2069,16 @@ bool AssetViewUtils::SyncPathsFromSourceControl(const TArray<FString>& Paths)
 	return SyncPathsFromSourceControl(Paths, /*bCheckDependencies=*/true);
 }
 
+bool AssetViewUtils::SyncRevisionFromSourceControl(const FString& Revision)
+{
+	return SyncPathsFromSourceControl(Revision, SourceControlHelpers::GetSourceControlLocations(), /*bCheckDependencies=*/false);
+}
+
+bool AssetViewUtils::SyncLatestFromSourceControl()
+{
+	return SyncRevisionFromSourceControl(TEXT(""));
+}
+
 void AssetViewUtils::ShowErrorNotifcation(const FText& InErrorMsg)
 {
 	if (!InErrorMsg.IsEmpty())
diff --git a/Engine/Source/Developer/AssetTools/Public/AssetViewUtils.h b/Engine/Source/Developer/AssetTools/Public/AssetViewUtils.h
--- a/Engine/Source/Developer/AssetTools/Public/AssetViewUtils.h
+++ b/Engine/Source/Developer/AssetTools/Public/AssetViewUtils.h
@@ -218,6 +218,9 @@ namespace AssetViewUtils
 	/** Syncs the specified paths from source control. */
 	ASSETTOOLS_API bool SyncPathsFromSourceControl(const TArray<FString>& Paths);
 
+	/** Syncs a specific revision from source control. */
+	ASSETTOOLS_API bool SyncRevisionFromSourceControl(const FString& Revision);
+
 	/** Syncs latest from source control. */
 	ASSETTOOLS_API bool SyncLatestFromSourceControl();
 
diff --git a/Engine/Source/Developer/SourceControl/Private/RevisionControlStyle/RevisionControlStyle.cpp b/Engine/Source/Developer/SourceControl/Private/RevisionControlStyle/RevisionControlStyle.cpp
--- a/Engine/Source/Developer/SourceControl/Private/RevisionControlStyle/RevisionControlStyle.cpp
+++ b/Engine/Source/Developer/SourceControl/Private/RevisionControlStyle/RevisionControlStyle.cpp
@@ -135,6 +135,7 @@ FDefaultRevisionControlStyle::FDefaultRevisionControlStyle() : FSlateStyleSet(St
 
 	// Revision Control States
 	Set("RevisionControl.VerticalLine", new IMAGE_BRUSH_SVG("Starship/SourceControl/RC_VerticalLine", CoreStyleConstants::Icon26x26));
+	Set("RevisionControl.VerticalLineStart", new IMAGE_BRUSH_SVG("Starship/SourceControl/RC_VerticalLineStart", CoreStyleConstants::Icon26x26));
 	Set("RevisionControl.VerticalLineDashed", new IMAGE_BRUSH_SVG("Starship/SourceControl/RC_VerticalLineDashed", CoreStyleConstants::Icon26x26));
 	Set("RevisionControl.CheckCircleLine", new IMAGE_BRUSH_SVG("Starship/SourceControl/RC_CheckCircleLine", CoreStyleConstants::Icon26x26));
 	Set("RevisionControl.LineCircle", new IMAGE_BRUSH_SVG("Starship/SourceControl/RC_LineCircle", CoreStyleConstants::Icon26x26));
diff --git a/Engine/Source/Developer/SourceControl/Private/SourceControlHelpers.cpp b/Engine/Source/Developer/SourceControl/Private/SourceControlHelpers.cpp
--- a/Engine/Source/Developer/SourceControl/Private/SourceControlHelpers.cpp
+++ b/Engine/Source/Developer/SourceControl/Private/SourceControlHelpers.cpp
@@ -1316,9 +1316,7 @@ bool USourceControlHelpers::RevertAllChangesAndReloadWorld()
 	TArray<FString> PackagesToReload;
 	ListRevertablePackages(PackagesToReload);
 
-	RevertAndReloadPackages(PackagesToReload, /*bRevertAll=*/true, /*bReloadWorld=*/true);
-
-	return true;
+	return RevertAndReloadPackages(PackagesToReload, /*bRevertAll=*/true, /*bReloadWorld=*/true);
 }
 
 bool USourceControlHelpers::RevertAndReloadPackages(const TArray<FString>& InPackagesToRevert, bool bRevertAll, bool bReloadWorld)
diff --git a/Engine/Source/Editor/SourceControlWindows/Private/SourceControlWindows.cpp b/Engine/Source/Editor/SourceControlWindows/Private/SourceControlWindows.cpp
--- a/Engine/Source/Editor/SourceControlWindows/Private/SourceControlWindows.cpp
+++ b/Engine/Source/Editor/SourceControlWindows/Private/SourceControlWindows.cpp
@@ -162,6 +162,11 @@ bool FSourceControlWindows::ShouldChoosePackagesToCheckBeVisible()
 }
 
 bool FSourceControlWindows::SyncLatest()
+{
+	return SyncRevision(TEXT(""));
+}
+
+bool FSourceControlWindows::SyncRevision(const FString& InRevision)
 {
 	bool bSaved = SaveDirtyPackages(/*bUseDialog=*/true);
 
@@ -179,7 +184,7 @@ bool FSourceControlWindows::SyncLatest()
 	// if properly saved or confirmation given, find all packages and use source control to update them.
 	if (bSaved)
 	{
-		bool bSuccess = AssetViewUtils::SyncLatestFromSourceControl();
+		bool bSuccess = AssetViewUtils::SyncRevisionFromSourceControl(InRevision);
 		if (!bSuccess)
 		{
 			FText Message(LOCTEXT("SCC_Sync_Failed", "Failed to sync files!"));
diff --git a/Engine/Source/Editor/SourceControlWindows/Public/SourceControlWindows.h b/Engine/Source/Editor/SourceControlWindows/Public/SourceControlWindows.h
--- a/Engine/Source/Editor/SourceControlWindows/Public/SourceControlWindows.h
+++ b/Engine/Source/Editor/SourceControlWindows/Public/SourceControlWindows.h
@@ -86,6 +86,13 @@ public:
 	 */
 	static bool SyncLatest();
 
+	/**
+	 * Saves all unsaved levels and assets and then - conditionally - performs an FSync operation to get the specified revision, and reloads the world
+	 *
+	 * @return	true - if command completed successfully.
+	 */
+	static bool SyncRevision(const FString& InRevision);
+
 	/** Determines whether we can sync to latest */
 	static bool CanSyncLatest();
 
