Fix for Header patching Actordefs

diff --git a/Engine/Source/Developer/AssetTools/Private/AssetHeaderPatcher.cpp b/Engine/Source/Developer/AssetTools/Private/AssetHeaderPatcher.cpp
--- a/Engine/Source/Developer/AssetTools/Private/AssetHeaderPatcher.cpp
+++ b/Engine/Source/Developer/AssetTools/Private/AssetHeaderPatcher.cpp
@@ -3,7 +3,6 @@
 #include "AssetHeaderPatcher.h"
 
 #include "Containers/ContainersFwd.h"
-#include "Algo/AnyOf.h"
 #include "UObject/PackageFileSummary.h"
 #include "UObject/NameTypes.h"
 #include "UObject/ObjectResource.h"
@@ -19,9 +18,14 @@
 #include "WorldPartition/WorldPartitionActorDescUtils.h"
 #include "AssetRegistry/AssetData.h"
 
+#include "Algo/AllOf.h"
+#include "Misc/Char.h"
+#include "Misc/Base64.h"
+#include "Serialization/MemoryWriter.h"
+
 DEFINE_LOG_CATEGORY_STATIC(LogAssetHeaderPatcher, Log, All);
 
-namespace 
+namespace
 {
 	// To override writing of FName's to ensure they have been patched
 	class FNamePatchingWriter final : public FArchiveProxy
@@ -32,8 +36,8 @@ namespace
 			, NameToIndexMap(InNameToIndexMap)
 		{
 		}
-		
-		virtual ~FNamePatchingWriter() 
+
+		virtual ~FNamePatchingWriter()
 		{
 		}
 
@@ -63,7 +67,7 @@ namespace
 		const TMap<FNameEntryId, int32>& NameToIndexMap;
 	};
 
-	enum class EPatchedSection 
+	enum class EPatchedSection
 	{
 		Summary,
 		NameTable,
@@ -136,9 +140,9 @@ namespace
 			return *this;
 		}
 
-		virtual FString GetArchiveName() const override 
+		virtual FString GetArchiveName() const override
 		{
-			return TEXT("FReadFNameAs2IntFromMemoryReader"); 
+			return TEXT("FReadFNameAs2IntFromMemoryReader");
 		}
 	private:
 		TArray<FName>& NameTable;
@@ -309,7 +313,7 @@ public:
 	TArray<FThumbnailEntry> ThumbnailTable;
 
 	// Asset registry data information
-	struct FAsetRegistryObjectData
+	struct FAssetRegistryObjectData
 	{
 		UE::AssetRegistry::FDeserializeObjectPackageData ObjectData;
 		TArray<UE::AssetRegistry::FDeserializeTagData> TagData;
@@ -319,7 +323,7 @@ public:
 	{
 		int64 SectionSize = -1;
 		UE::AssetRegistry::FDeserializePackageData PkgData;
-		TArray<FAsetRegistryObjectData> ObjectData;
+		TArray<FAssetRegistryObjectData> ObjectData;
 	};
 	FAsetRegistryData AsetRegistryData;
 };
@@ -336,7 +340,7 @@ UE::Tasks::TTask<FAssetHeaderPatcher::EResult> FAssetHeaderPatcher::Start(FStrin
 			{
 				UE_LOG(LogAssetHeaderPatcher, Error, TEXT("Failed to load %s"), *Self->SrcAsset);
 				return FAssetHeaderPatcherInner::EResult::ErrorFailedToLoadSourceAsset;
-			} 
+			}
 			else
 			{
 				return Self->PatchHeader();
@@ -487,7 +491,7 @@ FAssetHeaderPatcher::EResult FAssetHeaderPatcherInner::PatchHeader_Deserialize()
 
 		HeaderInformation.ImportTableSize = MemAr.Tell() - Summary.ImportOffset;
 	}
-	else 
+	else
 	{
 		HeaderInformation.ImportTableSize = 0;
 	}
@@ -565,7 +569,7 @@ FAssetHeaderPatcher::EResult FAssetHeaderPatcherInner::PatchHeader_Deserialize()
 		AsetRegistryData.ObjectData.Reserve(AsetRegistryData.PkgData.ObjectCount);
 		for (int32 i = 0; i < AsetRegistryData.PkgData.ObjectCount; ++i)
 		{
-			FAsetRegistryObjectData& ObjData = AsetRegistryData.ObjectData.Emplace_GetRef();
+			FAssetRegistryObjectData& ObjData = AsetRegistryData.ObjectData.Emplace_GetRef();
 			if (!ObjData.ObjectData.DoSerialize(MemAr, ErrorCode))
 			{
 				UE_LOG(LogAssetHeaderPatcher, Error, TEXT("Failed to deserialize asset registry data for %s"), *SrcAsset);
@@ -615,13 +619,12 @@ bool FAssetHeaderPatcherInner::DoPatch(FString& InOutString)
 			FStringView MaybeReplacement = Find(SearchAndReplace, FStringView(*InOutString, Idx));
 			if (!MaybeReplacement.IsEmpty())
 			{
-				FString Tmp = MaybeReplacement + InOutString.RightChop(Idx);;
-				InOutString = MoveTemp(Tmp);
+				InOutString = MaybeReplacement + InOutString.RightChop(Idx);
 				return true;
 			}
 		}
 	}
-	
+
 	{
 		// Patch quoted paths.
 		// Path occurs to the right of the first "'" 
@@ -631,8 +634,7 @@ bool FAssetHeaderPatcherInner::DoPatch(FString& InOutString)
 			FStringView MaybeReplacement = Find(SearchAndReplace, FStringView(*InOutString + Idx + 1, InOutString.Len() - (Idx + 2)));
 			if (!MaybeReplacement.IsEmpty())
 			{
-				FString Tmp = InOutString.LeftChop(Idx + 1) + MaybeReplacement + TCHAR('\'');
-				InOutString = MoveTemp(Tmp);
+				InOutString = InOutString.LeftChop(Idx + 1) + MaybeReplacement + TCHAR('\'');
 				return true;
 			}
 		}
@@ -656,11 +658,13 @@ bool FAssetHeaderPatcherInner::DoPatch(FName& InOutName)
 
 bool FAssetHeaderPatcherInner::DoPatch(FTopLevelAssetPath& InOutPath)
 {
-	FString Value = InOutPath.ToString();
-	bool bPatching = DoPatch(Value);
+	FName PackageName = InOutPath.GetPackageName();
+	FName AssetName = InOutPath.GetAssetName();
+	bool bPatching = DoPatch(PackageName);
+	bPatching |= DoPatch(AssetName);
 	if (bPatching)
 	{
-		InOutPath = Value;
+		InOutPath = FTopLevelAssetPath(PackageName, AssetName);
 	}
 	return bPatching;
 }
@@ -679,14 +683,14 @@ void FAssetHeaderPatcherInner::PatchHeader_PatchSections()
 		{
 			FName TmpName = Name;
 
-			if (DoPatch(TmpName)) 
+			if (DoPatch(TmpName))
 			{
 				FString TmpNameStr = Name.GetPlainNameString();
 				// If the string does not contain any path separators, then call it an Identifier.
 				if (!(TmpNameStr.Contains(TEXT("/")) || TmpNameStr.Contains(TEXT("\\"))))
 				{
 					ToAppend.Add(TmpName);
-				} 
+				}
 				else
 				{
 					Name = TmpName;
@@ -709,15 +713,10 @@ void FAssetHeaderPatcherInner::PatchHeader_PatchSections()
 	// Soft paths
 	for (FSoftObjectPath& PathRef : SoftObjectPathTable)
 	{
-		FTopLevelAssetPath AssetPath = PathRef.GetAssetPath();
-		FName PackageName = AssetPath.GetPackageName();
-		FName AssetName = AssetPath.GetAssetName();
-		bool bPatched = DoPatch(PackageName);
-		bPatched |= DoPatch(AssetName);
-		if (bPatched)
+		FTopLevelAssetPath TmpPath = PathRef.GetAssetPath();
+		if (DoPatch(TmpPath))
 		{
-			FTopLevelAssetPath NewAssetPath{ PackageName, AssetName };
-			PathRef.SetPath(NewAssetPath, PathRef.GetSubPathString());
+			PathRef.SetPath(TmpPath, PathRef.GetSubPathString());
 		}
 	}
 
@@ -752,55 +751,87 @@ void FAssetHeaderPatcherInner::PatchHeader_PatchSections()
 	}
 
 	// Asset Register Data
-	for (FAsetRegistryObjectData& ObjData : AsetRegistryData.ObjectData)
+	for (FAssetRegistryObjectData& ObjData : AsetRegistryData.ObjectData)
 	{
 		DoPatch(ObjData.ObjectData.ObjectPath);
-		DoPatch(ObjData.ObjectData.ObjectClassName);
 
-		bool bHasActor = false;
+		{	// ObjData.ObjectData.ObjectClassName is stored as a string, but is used as a FTopLevelAssetPath
+			FTopLevelAssetPath Tmp(ObjData.ObjectData.ObjectClassName);
+			if (DoPatch(Tmp)) 
+			{
+				FString TS = Tmp.ToString();
+				ObjData.ObjectData.ObjectClassName = TS;
+			}
+		}
 
 		for (UE::AssetRegistry::FDeserializeTagData& TagData : ObjData.TagData)
 		{
-			bHasActor |= (TagData.Key == FWorldPartitionActorDescUtils::ActorMetaDataClassTagName()
-				       || TagData.Key == FWorldPartitionActorDescUtils::ActorMetaDataTagName());
-			
-			DoPatch(TagData.Value);
-		}
-
-		if (bHasActor)
-		{ 
-			FGCScopeGuard GCGuard;
-
-			const FString LongPackageName(SrcAsset);
-			const FString ObjectPath(ObjData.ObjectData.ObjectPath);
-			const FTopLevelAssetPath AssetClassPathName(ObjData.ObjectData.ObjectClassName);
-			const FAssetDataTagMap Tags(MakeTagMap(ObjData.TagData));
-			const FAssetData AssetData(LongPackageName, ObjectPath, AssetClassPathName, Tags);
-
-			if (TUniquePtr<FWorldPartitionActorDesc> ActorDesc = FWorldPartitionActorDescUtils::GetActorDescriptorFromAssetData(AssetData))
+			if (TagData.Key == FWorldPartitionActorDescUtils::ActorMetaDataTagName())
 			{
-				bool bPatchedActorDesc = DoPatch(ActorDesc->BaseClass);
-
-				if (ActorDesc->bIsUsingDataLayerAsset)
+				// This code is to be revisited (UE-201639)
+				// 
+				// Actor meta data is a binary blob from a advanced Archive, then bin64 encoded.
+				// The Archive is advanced as it does delta encoded from a base instance of a class.
+				// When we are copying at this point, we dont have that loaded.
+				// So instead, We search through the decoded information, looking for strings (FNames, SoftPaths)
+				// which in this for are a 4byte int length, and then the bytes. 
+				// Assumptions: 
+				//   1. The strings I'm looking for will be ascii.
+				//   2. The int is stored in native byte order.
+				TArray<uint8> ActorMetaSrc;
+				verify(FBase64::Decode(TagData.Value, ActorMetaSrc));
+				int32 SrcLen = ActorMetaSrc.Num();
+				const uint8* Src = ActorMetaSrc.GetData();
+				static_assert(sizeof(uint8) == sizeof(ANSICHAR));
+
+				TArray<uint8> ActorMetaDst;
+				ActorMetaDst.Reserve(SrcLen);
+				FMemoryWriter DstAr(ActorMetaDst, /* bIsPersistent*/ true);
+
+				bool bPatched = false;
+
+				for (int32 I = 0; I < SrcLen; )
 				{
-					for (FName& DataLayerAssetPath : ActorDesc->DataLayers)
+					// 5 bytes is the min size for a string that could be interesting.
+					// 4 for the int32 size, and 1 character.
+					if ((I + 5) < SrcLen)
 					{
-						bPatchedActorDesc |= DoPatch(DataLayerAssetPath);
+						int32 PossibleLen;
+						memcpy(&PossibleLen, Src + I, sizeof(PossibleLen));
+						if (PossibleLen > 0 && ((I + (int32)sizeof(PossibleLen) + (int64)PossibleLen) < (int64)SrcLen)) // using 64bit math to guard against overflow
+						{
+							// Length looks okay.. Test to ensure it has no control characters.
+							FString PossibleStr((const ANSICHAR*)(Src + I + sizeof(PossibleLen)), PossibleLen);
+							bool bIsPrintable = Algo::AllOf(PossibleStr, [](TCHAR C) { return FChar::IsPrint(C); });
+
+							if (bIsPrintable)
+							{
+								if (DoPatch(PossibleStr))
+								{
+									bPatched = true;
+
+									DstAr << PossibleStr;
+									I += (int32)sizeof(PossibleLen) + PossibleLen;
+									continue;
+								}
+							}
+						}
 					}
+
+					// Copy 1 byte
+					DstAr.Serialize((void*)(Src + I), 1);
+					++I;
 				}
 
-				if (bPatchedActorDesc)
+				if (bPatched)
 				{
-					for (UE::AssetRegistry::FDeserializeTagData& Tag : ObjData.TagData)
-					{
-						if (Tag.Key == FWorldPartitionActorDescUtils::ActorMetaDataTagName())
-						{
-							Tag.Value = FWorldPartitionActorDescUtils::GetAssetDataFromActorDescriptor(ActorDesc);
-							break;
-						}
-					}
+					TagData.Value = FBase64::Encode(ActorMetaDst);
 				}
 			}
+			else 
+			{
+				DoPatch(TagData.Value);
+			}
 		}
 	}
 }
@@ -874,14 +905,14 @@ FAssetHeaderPatcher::EResult FAssetHeaderPatcherInner::PatchHeader_WriteDestinat
 	}
 
 	int64 LastSectionEndedAt = 0;
-	
+
 	for (int SectionIdx = 0; SectionIdx < UE_ARRAY_COUNT(SourceSections); ++SectionIdx)
 	{
 		const FSectionData& SourceSection = SourceSections[SectionIdx];
 
 		// skip processing empty non required chunks.
 		// really the only option section is the Thumbnails, and its annoying its ion the middle.
-		if (!SourceSection.bRequired && SourceSection.Size <= 0) 
+		if (!SourceSection.bRequired && SourceSection.Size <= 0)
 		{
 			continue;
 		}
@@ -898,175 +929,175 @@ FAssetHeaderPatcher::EResult FAssetHeaderPatcherInner::PatchHeader_WriteDestinat
 		// Serialize the current patched section and patch summary offsets
 		switch (SourceSection.Section)
 		{
-			case EPatchedSection::Summary:
+		case EPatchedSection::Summary:
+		{
+			// We will write the Summary twice.
+			// The first is so we get its new size (if the name was changed in patching)
+			// The second is done after the loop, to patch up all the offsets.
+			check(Writer.Tell() == 0);
+			Writer << Summary;
+			const int64 SummarySize = Writer.Tell();
+			const int64 Delta = SummarySize - SourceSection.Size;
+			PatchSummaryOffsets(Summary, 0, Delta);
+			Summary.TotalHeaderSize += (int32)Delta;
+
+			break;
+		}
+
+		case EPatchedSection::NameTable:
+		{
+			const int64 NameTableStartOffset = Writer.Tell();
+			for (FName& Name : NameTable)
 			{
-				// We will write the Summary twice.
-				// The first is so we get its new size (if the name was changed in patching)
-				// The second is done after the loop, to patch up all the offsets.
-				check(Writer.Tell() == 0);
-				Writer << Summary;
-				const int64 SummarySize = Writer.Tell();
-				const int64 Delta = SummarySize - SourceSection.Size;
-				PatchSummaryOffsets(Summary, 0, Delta);
-				Summary.TotalHeaderSize += (int32)Delta;
-
-				break;
+				const FNameEntry* Entry = FName::GetEntry(Name.GetDisplayIndex());
+				check(Entry);
+				Entry->Write(Writer);
 			}
+			const int64 NameTableSize = Writer.Tell() - NameTableStartOffset;
+			const int64 Delta = NameTableSize - SourceSection.Size;
+			PatchSummaryOffsets(Summary, NameTableStartOffset, Delta);
+			Summary.TotalHeaderSize += (int32)Delta;
+			check(Summary.NameCount == NameTable.Num());
+			check(Summary.NameOffset == NameTableStartOffset);
+
+			break;
+		}
 
-			case EPatchedSection::NameTable:
+		case EPatchedSection::SoftPathTable:
+		{
+			const int64 TableStartOffset = Writer.Tell();
+			for (FSoftObjectPath& PathRef : SoftObjectPathTable)
 			{
-				const int64 NameTableStartOffset = Writer.Tell();
-				for (FName& Name : NameTable)
-				{
-					const FNameEntry* Entry = FName::GetEntry(Name.GetDisplayIndex());
-					check(Entry);
-					Entry->Write(Writer);
-				}
-				const int64 NameTableSize = Writer.Tell() - NameTableStartOffset;
-				const int64 Delta = NameTableSize - SourceSection.Size;
-				PatchSummaryOffsets(Summary, NameTableStartOffset, Delta);
-				Summary.TotalHeaderSize += (int32)Delta;
-				check(Summary.NameCount == NameTable.Num());
-				check(Summary.NameOffset == NameTableStartOffset);
-
-				break;
+				PathRef.SerializePath(Writer);
 			}
+			const int64 TableSize = Writer.Tell() - TableStartOffset;
+			const int64 Delta = TableSize - SourceSection.Size;
+			checkf(Delta == 0, TEXT("Delta should be Zero. is %d"), (int)Delta);
+			check(Summary.SoftObjectPathsCount == SoftObjectPathTable.Num());
+			check(Summary.SoftObjectPathsOffset == TableStartOffset);
 
-			case EPatchedSection::SoftPathTable:
-			{
-				const int64 TableStartOffset = Writer.Tell();
-				for (FSoftObjectPath& PathRef : SoftObjectPathTable)
-				{
-					PathRef.SerializePath(Writer);
-				}
-				const int64 TableSize = Writer.Tell() - TableStartOffset;
-				const int64 Delta = TableSize - SourceSection.Size;
-				checkf(Delta == 0, TEXT("Delta should be Zero. is %d"), (int)Delta);
-				check(Summary.SoftObjectPathsCount == SoftObjectPathTable.Num());
-				check(Summary.SoftObjectPathsOffset == TableStartOffset);
+			break;
+		}
 
-				break;
+		case EPatchedSection::ImportTable:
+		{
+			const int64 ImportTableStartOffset = Writer.Tell();
+			for (FObjectImport& Import : ImportTable)
+			{
+				Writer << Import;
 			}
+			const int64 ImportTableSize = Writer.Tell() - ImportTableStartOffset;
+			const int64 Delta = ImportTableSize - SourceSection.Size;
+			check(Delta == 0);
+			checkf(ImportTableSize == SourceSection.Size, TEXT("%d == %d"), (int)ImportTableSize, (int)SourceSection.Size); // We only patch export table offsets, we should not be patching size
+			checkf(Summary.ImportCount == ImportTable.Num(), TEXT("%d == %d"), Summary.ImportCount, ImportTable.Num());
+			checkf(Summary.ImportOffset == ImportTableStartOffset, TEXT("%d == %d"), Summary.ImportOffset, ImportTableStartOffset);
+
+			break;
+		}
 
-			case EPatchedSection::ImportTable:
+		case EPatchedSection::ExportTable:
+		{
+			// The export table offsets aren't correct yet.
+			// Once we know them, we will seek back and write it a second time.
+			const int64 ExportTableStartOffset = Writer.Tell();
+			for (FObjectExport& Export : ExportTable)
 			{
-				const int64 ImportTableStartOffset = Writer.Tell();
-				for (FObjectImport& Import : ImportTable)
-				{
-					Writer << Import;
-				}
-				const int64 ImportTableSize = Writer.Tell() - ImportTableStartOffset;
-				const int64 Delta = ImportTableSize - SourceSection.Size;
-				check(Delta == 0);
-				checkf(ImportTableSize == SourceSection.Size, TEXT("%d == %d"), (int)ImportTableSize, (int)SourceSection.Size); // We only patch export table offsets, we should not be patching size
-				checkf(Summary.ImportCount == ImportTable.Num(), TEXT("%d == %d"), Summary.ImportCount, ImportTable.Num());
-				checkf(Summary.ImportOffset == ImportTableStartOffset, TEXT("%d == %d"), Summary.ImportOffset, ImportTableStartOffset);
-
-				break;
+				Writer << Export;
 			}
+			const int64 ExportTableSize = Writer.Tell() - ExportTableStartOffset;
+			const int64 Delta = ExportTableSize - SourceSection.Size;
+			check(Delta == 0);
+			checkf(ExportTableSize == SourceSection.Size, TEXT("%d == %d"), (int)ExportTableSize, (int)SourceSection.Size); // We only patch export table offsets, we should not be patching size
+			checkf(Summary.ExportCount == ExportTable.Num(), TEXT("%d == %d"), Summary.ExportCount, ExportTable.Num());
+			checkf(Summary.ExportOffset == ExportTableStartOffset, TEXT("%d == %d"), Summary.ExportOffset, ExportTableStartOffset);
+
+			break;
+		}
 
-			case EPatchedSection::ExportTable:
+		case EPatchedSection::SoftPackageReferencesTable:
+		{
+			const int64 TableStartOffset = Writer.Tell();
+			for (FName& Reference : SoftPackageReferencesTable)
 			{
-				// The export table offsets aren't correct yet.
-				// Once we know them, we will seek back and write it a second time.
-				const int64 ExportTableStartOffset = Writer.Tell();
-				for (FObjectExport& Export : ExportTable)
-				{
-					Writer << Export;
-				}
-				const int64 ExportTableSize = Writer.Tell() - ExportTableStartOffset;
-				const int64 Delta = ExportTableSize - SourceSection.Size;
-				check(Delta == 0);
-				checkf(ExportTableSize == SourceSection.Size, TEXT("%d == %d"), (int)ExportTableSize, (int)SourceSection.Size); // We only patch export table offsets, we should not be patching size
-				checkf(Summary.ExportCount == ExportTable.Num(), TEXT("%d == %d"), Summary.ExportCount, ExportTable.Num());
-				checkf(Summary.ExportOffset == ExportTableStartOffset, TEXT("%d == %d"), Summary.ExportOffset, ExportTableStartOffset);
-
-				break;
+				Writer << Reference;
 			}
+			const int64 TableSize = Writer.Tell() - TableStartOffset;
+			const int64 Delta = TableSize - SourceSection.Size;
+			checkf(Delta == 0, TEXT("Delta should be Zero. is %d"), (int)Delta);
+			check(Summary.SoftPackageReferencesCount == SoftPackageReferencesTable.Num());
+			check(Summary.SoftPackageReferencesOffset == TableStartOffset);
 
-			case EPatchedSection::SoftPackageReferencesTable:
-			{
-				const int64 TableStartOffset = Writer.Tell();
-				for (FName& Reference : SoftPackageReferencesTable)
-				{
-					Writer << Reference;
-				}
-				const int64 TableSize = Writer.Tell() - TableStartOffset;
-				const int64 Delta = TableSize - SourceSection.Size;
-				checkf(Delta == 0, TEXT("Delta should be Zero. is %d"), (int)Delta);
-				check(Summary.SoftPackageReferencesCount == SoftPackageReferencesTable.Num());
-				check(Summary.SoftPackageReferencesOffset == TableStartOffset);
+			break;
+		}
 
-				break;
+		case EPatchedSection::ThumbnailTable:
+		{
+			// Luckily the strings in the thumbnail table don't include the package, so only the offsets need to be patched
+			const int64 ThumbnailTableStartOffset = Writer.Tell();
+			const int64 ThumbnailTableDeltaOffset = ThumbnailTableStartOffset - SourceSection.Offset;
+			int32 ThumbnailCount = ThumbnailTable.Num();
+			Writer << ThumbnailCount;
+			for (FThumbnailEntry& Entry : ThumbnailTable)
+			{
+				Writer << Entry.ObjectShortClassName;
+				Writer << Entry.ObjectPathWithoutPackageName;
+				Entry.FileOffset += (int32)ThumbnailTableDeltaOffset; // Thumbnail payloads immediately follow the table, so we can just apply the section offset delta here
+				Writer << Entry.FileOffset;
 			}
+			const int64 ThumbnailTableSize = Writer.Tell() - ThumbnailTableStartOffset;
+			checkf(ThumbnailTableStartOffset == Summary.ThumbnailTableOffset, TEXT("%zd == %zd"), ThumbnailTableStartOffset, Summary.ThumbnailTableOffset);
+			check(ThumbnailTableSize == SourceSection.Size); // We only patch thumbnail table offsets, we should not be patching size
 
-			case EPatchedSection::ThumbnailTable:
-			{
-				// Luckily the strings in the thumbnail table don't include the package, so only the offsets need to be patched
-				const int64 ThumbnailTableStartOffset = Writer.Tell();
-				const int64 ThumbnailTableDeltaOffset = ThumbnailTableStartOffset - SourceSection.Offset;
-				int32 ThumbnailCount = ThumbnailTable.Num();
-				Writer << ThumbnailCount;
-				for (FThumbnailEntry& Entry : ThumbnailTable)
-				{
-					Writer << Entry.ObjectShortClassName;
-					Writer << Entry.ObjectPathWithoutPackageName;
-					Entry.FileOffset += (int32)ThumbnailTableDeltaOffset; // Thumbnail payloads immediately follow the table, so we can just apply the section offset delta here
-					Writer << Entry.FileOffset;
-				}
-				const int64 ThumbnailTableSize = Writer.Tell() - ThumbnailTableStartOffset;
-				checkf(ThumbnailTableStartOffset == Summary.ThumbnailTableOffset, TEXT("%zd == %zd"), ThumbnailTableStartOffset, Summary.ThumbnailTableOffset);
-				check(ThumbnailTableSize == SourceSection.Size); // We only patch thumbnail table offsets, we should not be patching size
+			break;
+		}
 
-				break;
-			}
+		case EPatchedSection::AssetRegistryData:
+		{
+			const int64 AsetRegistryDataStartOffset = Writer.Tell();
+			checkf(AsetRegistryDataStartOffset == Summary.AssetRegistryDataOffset, TEXT("%zd == %zd"), AsetRegistryDataStartOffset, Summary.AssetRegistryDataOffset);
 
-			case EPatchedSection::AssetRegistryData:
+			// Manually write this back out, there isn't a nicely factored function to call for this
+			if (AsetRegistryData.PkgData.DependencyDataOffset != INDEX_NONE)
 			{
-				const int64 AsetRegistryDataStartOffset = Writer.Tell();
-				checkf(AsetRegistryDataStartOffset == Summary.AssetRegistryDataOffset, TEXT("%zd == %zd"), AsetRegistryDataStartOffset, Summary.AssetRegistryDataOffset);
+				Writer << AsetRegistryData.PkgData.DependencyDataOffset;
+			}
+			Writer << AsetRegistryData.PkgData.ObjectCount;
 
-				// Manually write this back out, there isn't a nicely factored function to call for this
-				if (AsetRegistryData.PkgData.DependencyDataOffset != INDEX_NONE)
-				{
-					Writer << AsetRegistryData.PkgData.DependencyDataOffset;
-				}
-				Writer << AsetRegistryData.PkgData.ObjectCount;
+			check(AsetRegistryData.PkgData.ObjectCount == AsetRegistryData.ObjectData.Num());
+			for (FAssetRegistryObjectData& ObjData : AsetRegistryData.ObjectData)
+			{
+				Writer << ObjData.ObjectData.ObjectPath;
+				Writer << ObjData.ObjectData.ObjectClassName;
+				Writer << ObjData.ObjectData.TagCount;
 
-				check(AsetRegistryData.PkgData.ObjectCount == AsetRegistryData.ObjectData.Num());
-				for (FAsetRegistryObjectData& ObjData : AsetRegistryData.ObjectData)
+				check(ObjData.ObjectData.TagCount == ObjData.TagData.Num());
+				for (UE::AssetRegistry::FDeserializeTagData& TagData : ObjData.TagData)
 				{
-					Writer << ObjData.ObjectData.ObjectPath;
-					Writer << ObjData.ObjectData.ObjectClassName;
-					Writer << ObjData.ObjectData.TagCount;
-
-					check(ObjData.ObjectData.TagCount == ObjData.TagData.Num());
-					for (UE::AssetRegistry::FDeserializeTagData& TagData : ObjData.TagData)
-					{
-						Writer << TagData.Key;
-						Writer << TagData.Value;
-					}
+					Writer << TagData.Key;
+					Writer << TagData.Value;
 				}
+			}
 
-				const int64 AsetRegistryDataSize = Writer.Tell() - AsetRegistryDataStartOffset;
-				const int64 Delta = AsetRegistryDataSize - SourceSection.Size;
-				PatchSummaryOffsets(Summary, AsetRegistryDataStartOffset, Delta);
-				Summary.TotalHeaderSize += (int32)Delta;
+			const int64 AsetRegistryDataSize = Writer.Tell() - AsetRegistryDataStartOffset;
+			const int64 Delta = AsetRegistryDataSize - SourceSection.Size;
+			PatchSummaryOffsets(Summary, AsetRegistryDataStartOffset, Delta);
+			Summary.TotalHeaderSize += (int32)Delta;
 
-				if (AsetRegistryData.PkgData.DependencyDataOffset != INDEX_NONE)
-				{
-					// DependencyDataOffset is not relative but points to just after the rest of the AR data
-					// We will seek back and write this later
-					const int64 DependencyDataDelta = AsetRegistryDataStartOffset - SourceSection.Offset + Delta;
-					AsetRegistryData.PkgData.DependencyDataOffset += DependencyDataDelta;
-				}
-
-				break;
+			if (AsetRegistryData.PkgData.DependencyDataOffset != INDEX_NONE)
+			{
+				// DependencyDataOffset is not relative but points to just after the rest of the AR data
+				// We will seek back and write this later
+				const int64 DependencyDataDelta = AsetRegistryDataStartOffset - SourceSection.Offset + Delta;
+				AsetRegistryData.PkgData.DependencyDataOffset += DependencyDataDelta;
 			}
 
-			default:
-				UE_LOG(LogAssetHeaderPatcher, Error, TEXT("Unexpected section for %s"), *SrcAsset);
-				return EResult::ErrorUnkownSection;
+			break;
+		}
+
+		default:
+			UE_LOG(LogAssetHeaderPatcher, Error, TEXT("Unexpected section for %s"), *SrcAsset);
+			return EResult::ErrorUnkownSection;
 		}
 
 		if (Writer.IsError())
diff --git a/Engine/Source/Developer/AssetTools/Private/AssetTools.cpp b/Engine/Source/Developer/AssetTools/Private/AssetTools.cpp
--- a/Engine/Source/Developer/AssetTools/Private/AssetTools.cpp
+++ b/Engine/Source/Developer/AssetTools/Private/AssetTools.cpp
@@ -2245,6 +2245,7 @@ TMap<FString, FString> GenerateAdditionalAssetMappings(const TMap<FString, FStri
 		if (SrcPackageName != DstPackageName)
 		{
 			Result.Add({ FString(SrcPackageName), FString(DstPackageName) });
+			Result.Add({ FString(SrcPackageName) + TEXT("_C"), FString(DstPackageName) + TEXT("_C") }); // catch compiled blueprint names
 		}
 	}
 
