World Partition Actor Descriptor Patching: it is now possible to patch actor descriptors without having to load/save them when migrating projects, through a custom archive that will write out exactly what's read, with support for patching names and soft object paths. This allow migrating projects without having to enable required plugins or load required classes, etc.

diff --git a/Engine/Source/Developer/AssetTools/Private/AssetHeaderPatcher.cpp b/Engine/Source/Developer/AssetTools/Private/AssetHeaderPatcher.cpp
--- a/Engine/Source/Developer/AssetTools/Private/AssetHeaderPatcher.cpp
+++ b/Engine/Source/Developer/AssetTools/Private/AssetHeaderPatcher.cpp
@@ -18,11 +18,6 @@
 #include "WorldPartition/WorldPartitionActorDescUtils.h"
 #include "AssetRegistry/AssetData.h"
 
-#include "Algo/AllOf.h"
-#include "Misc/Char.h"
-#include "Misc/Base64.h"
-#include "Serialization/MemoryWriter.h"
-
 DEFINE_LOG_CATEGORY_STATIC(LogAssetHeaderPatcher, Log, All);
 
 namespace
@@ -268,6 +263,7 @@ public:
 
 	bool DoPatch(FString& InOutString);
 	bool DoPatch(FName& InOutName);
+	bool DoPatch(FSoftObjectPath& InOutSoft);
 	bool DoPatch(FTopLevelAssetPath& InOutPath);
 
 	FAssetHeaderPatcher::EResult PatchHeader();
@@ -319,13 +315,13 @@ public:
 		TArray<UE::AssetRegistry::FDeserializeTagData> TagData;
 	};
 
-	struct FAsetRegistryData
+	struct FAssetRegistryData
 	{
 		int64 SectionSize = -1;
 		UE::AssetRegistry::FDeserializePackageData PkgData;
 		TArray<FAssetRegistryObjectData> ObjectData;
 	};
-	FAsetRegistryData AsetRegistryData;
+	FAssetRegistryData AssetRegistryData;
 };
 
 
@@ -560,16 +556,16 @@ FAssetHeaderPatcher::EResult FAssetHeaderPatcherInner::PatchHeader_Deserialize()
 		MemAr.Seek(Summary.AssetRegistryDataOffset);
 
 		UE::AssetRegistry::EReadPackageDataMainErrorCode ErrorCode;
-		if (!AsetRegistryData.PkgData.DoSerialize(MemAr, Summary, ErrorCode))
+		if (!AssetRegistryData.PkgData.DoSerialize(MemAr, Summary, ErrorCode))
 		{
 			UE_LOG(LogAssetHeaderPatcher, Error, TEXT("Failed to deserialize asset registry data for %s"), *SrcAsset);
 			return EResult::ErrorFailedToDeserializeSourceAsset;
 		}
 
-		AsetRegistryData.ObjectData.Reserve(AsetRegistryData.PkgData.ObjectCount);
-		for (int32 i = 0; i < AsetRegistryData.PkgData.ObjectCount; ++i)
+		AssetRegistryData.ObjectData.Reserve(AssetRegistryData.PkgData.ObjectCount);
+		for (int32 i = 0; i < AssetRegistryData.PkgData.ObjectCount; ++i)
 		{
-			FAssetRegistryObjectData& ObjData = AsetRegistryData.ObjectData.Emplace_GetRef();
+			FAssetRegistryObjectData& ObjData = AssetRegistryData.ObjectData.Emplace_GetRef();
 			if (!ObjData.ObjectData.DoSerialize(MemAr, ErrorCode))
 			{
 				UE_LOG(LogAssetHeaderPatcher, Error, TEXT("Failed to deserialize asset registry data for %s"), *SrcAsset);
@@ -587,7 +583,7 @@ FAssetHeaderPatcher::EResult FAssetHeaderPatcherInner::PatchHeader_Deserialize()
 			}
 		}
 
-		AsetRegistryData.SectionSize = MemAr.Tell() - Summary.AssetRegistryDataOffset;
+		AssetRegistryData.SectionSize = MemAr.Tell() - Summary.AssetRegistryDataOffset;
 	}
 
 	return EResult::Success;
@@ -646,14 +642,25 @@ bool FAssetHeaderPatcherInner::DoPatch(FString& InOutString)
 bool FAssetHeaderPatcherInner::DoPatch(FName& InOutName)
 {
 	FString Value = InOutName.GetPlainNameString();
-	bool bPatching = DoPatch(Value);
-	if (bPatching)
+	if (DoPatch(Value))
 	{
 		// Use the same Number as the original Name for consistency.
 		// Otherwise different files with the same name, generate different numbers, and this breaks linking.
 		InOutName = FName(Value, InOutName.GetNumber());
+		return true;
 	}
-	return bPatching;
+	return false;
+}
+
+bool FAssetHeaderPatcherInner::DoPatch(FSoftObjectPath& InOutSoft)
+{
+	FTopLevelAssetPath TmpPath = InOutSoft.GetAssetPath();
+	if (DoPatch(TmpPath))
+	{
+		InOutSoft.SetPath(TmpPath, InOutSoft.GetSubPathString());
+		return true;
+	}
+	return false;
 }
 
 bool FAssetHeaderPatcherInner::DoPatch(FTopLevelAssetPath& InOutPath)
@@ -713,11 +720,7 @@ void FAssetHeaderPatcherInner::PatchHeader_PatchSections()
 	// Soft paths
 	for (FSoftObjectPath& PathRef : SoftObjectPathTable)
 	{
-		FTopLevelAssetPath TmpPath = PathRef.GetAssetPath();
-		if (DoPatch(TmpPath))
-		{
-			PathRef.SetPath(TmpPath, PathRef.GetSubPathString());
-		}
+		DoPatch(PathRef);
 	}
 
 	// Import table
@@ -751,7 +754,7 @@ void FAssetHeaderPatcherInner::PatchHeader_PatchSections()
 	}
 
 	// Asset Register Data
-	for (FAssetRegistryObjectData& ObjData : AsetRegistryData.ObjectData)
+	for (FAssetRegistryObjectData& ObjData : AssetRegistryData.ObjectData)
 	{
 		DoPatch(ObjData.ObjectData.ObjectPath);
 
@@ -768,64 +771,27 @@ void FAssetHeaderPatcherInner::PatchHeader_PatchSections()
 		{
 			if (TagData.Key == FWorldPartitionActorDescUtils::ActorMetaDataTagName())
 			{
-				// This code is to be revisited (UE-201639)
-				// 
-				// Actor meta data is a binary blob from a advanced Archive, then bin64 encoded.
-				// The Archive is advanced as it does delta encoded from a base instance of a class.
-				// When we are copying at this point, we dont have that loaded.
-				// So instead, We search through the decoded information, looking for strings (FNames, SoftPaths)
-				// which in this for are a 4byte int length, and then the bytes. 
-				// Assumptions: 
-				//   1. The strings I'm looking for will be ascii.
-				//   2. The int is stored in native byte order.
-				TArray<uint8> ActorMetaSrc;
-				verify(FBase64::Decode(TagData.Value, ActorMetaSrc));
-				int32 SrcLen = ActorMetaSrc.Num();
-				const uint8* Src = ActorMetaSrc.GetData();
-				static_assert(sizeof(uint8) == sizeof(ANSICHAR));
-
-				TArray<uint8> ActorMetaDst;
-				ActorMetaDst.Reserve(SrcLen);
-				FMemoryWriter DstAr(ActorMetaDst, /* bIsPersistent*/ true);
-
-				bool bPatched = false;
-
-				for (int32 I = 0; I < SrcLen; )
-				{
-					// 5 bytes is the min size for a string that could be interesting.
-					// 4 for the int32 size, and 1 character.
-					if ((I + 5) < SrcLen)
-					{
-						int32 PossibleLen;
-						memcpy(&PossibleLen, Src + I, sizeof(PossibleLen));
-						if (PossibleLen > 0 && ((I + (int32)sizeof(PossibleLen) + (int64)PossibleLen) < (int64)SrcLen)) // using 64bit math to guard against overflow
-						{
-							// Length looks okay.. Test to ensure it has no control characters.
-							FString PossibleStr((const ANSICHAR*)(Src + I + sizeof(PossibleLen)), PossibleLen);
-							bool bIsPrintable = Algo::AllOf(PossibleStr, [](TCHAR C) { return FChar::IsPrint(C); });
-
-							if (bIsPrintable)
-							{
-								if (DoPatch(PossibleStr))
-								{
-									bPatched = true;
-
-									DstAr << PossibleStr;
-									I += (int32)sizeof(PossibleLen) + PossibleLen;
-									continue;
-								}
-							}
-						}
-					}
-
-					// Copy 1 byte
-					DstAr.Serialize((void*)(Src + I), 1);
-					++I;
-				}
+				const FString LongPackageName(SrcAsset);
+				const FString ObjectPath(ObjData.ObjectData.ObjectPath);
+				const FTopLevelAssetPath AssetClassPathName(ObjData.ObjectData.ObjectClassName);
+				const FAssetDataTagMap Tags(MakeTagMap(ObjData.TagData));
+				const FAssetData AssetData(LongPackageName, ObjectPath, AssetClassPathName, Tags);
 
-				if (bPatched)
+				struct FWorldPartitionAssetDataPatcherInner : FWorldPartitionAssetDataPatcher
+				{
+					FWorldPartitionAssetDataPatcherInner(FAssetHeaderPatcherInner* InInner) : Inner(InInner) {}
+					virtual bool DoPatch(FString& InOutString) override { return Inner->DoPatch(InOutString); }
+					virtual bool DoPatch(FName& InOutName) override { return Inner->DoPatch(InOutName); }
+					virtual bool DoPatch(FSoftObjectPath& InOutSoft) override { return Inner->DoPatch(InOutSoft); }
+					virtual bool DoPatch(FTopLevelAssetPath& InOutPath) override { return Inner->DoPatch(InOutPath); }
+					FAssetHeaderPatcherInner* Inner;
+				};
+				
+				FString PatchedAssetData;
+				FWorldPartitionAssetDataPatcherInner Patcher(this);
+				if (FWorldPartitionActorDescUtils::GetPatchedAssetDataFromAssetData(AssetData, PatchedAssetData, &Patcher))
 				{
-					TagData.Value = FBase64::Encode(ActorMetaDst);
+					TagData.Value = PatchedAssetData;
 				}
 			}
 			else 
@@ -849,7 +815,7 @@ FAssetHeaderPatcher::EResult FAssetHeaderPatcherInner::PatchHeader_WriteDestinat
 		{ EPatchedSection::ExportTable,					Summary.ExportOffset,					HeaderInformation.ExportTableSize,		true },
 		{ EPatchedSection::SoftPackageReferencesTable,	Summary.SoftPackageReferencesOffset,	HeaderInformation.SoftPackageReferencesListSize, false },
 		{ EPatchedSection::ThumbnailTable,				Summary.ThumbnailTableOffset,			HeaderInformation.ThumbnailTableSize,	false },
-		{ EPatchedSection::AssetRegistryData,			Summary.AssetRegistryDataOffset,		AsetRegistryData.SectionSize,			true },
+		{ EPatchedSection::AssetRegistryData,			Summary.AssetRegistryDataOffset,		AssetRegistryData.SectionSize,			true },
 	};
 
 	const int32 SourceTotalHeaderSize = Summary.TotalHeaderSize;
@@ -1054,18 +1020,18 @@ FAssetHeaderPatcher::EResult FAssetHeaderPatcherInner::PatchHeader_WriteDestinat
 
 		case EPatchedSection::AssetRegistryData:
 		{
-			const int64 AsetRegistryDataStartOffset = Writer.Tell();
-			checkf(AsetRegistryDataStartOffset == Summary.AssetRegistryDataOffset, TEXT("%zd == %zd"), AsetRegistryDataStartOffset, Summary.AssetRegistryDataOffset);
+			const int64 AssetRegistryDataStartOffset = Writer.Tell();
+			checkf(AssetRegistryDataStartOffset == Summary.AssetRegistryDataOffset, TEXT("%zd == %zd"), AssetRegistryDataStartOffset, Summary.AssetRegistryDataOffset);
 
 			// Manually write this back out, there isn't a nicely factored function to call for this
-			if (AsetRegistryData.PkgData.DependencyDataOffset != INDEX_NONE)
+			if (AssetRegistryData.PkgData.DependencyDataOffset != INDEX_NONE)
 			{
-				Writer << AsetRegistryData.PkgData.DependencyDataOffset;
+				Writer << AssetRegistryData.PkgData.DependencyDataOffset;
 			}
-			Writer << AsetRegistryData.PkgData.ObjectCount;
+			Writer << AssetRegistryData.PkgData.ObjectCount;
 
-			check(AsetRegistryData.PkgData.ObjectCount == AsetRegistryData.ObjectData.Num());
-			for (FAssetRegistryObjectData& ObjData : AsetRegistryData.ObjectData)
+			check(AssetRegistryData.PkgData.ObjectCount == AssetRegistryData.ObjectData.Num());
+			for (FAssetRegistryObjectData& ObjData : AssetRegistryData.ObjectData)
 			{
 				Writer << ObjData.ObjectData.ObjectPath;
 				Writer << ObjData.ObjectData.ObjectClassName;
@@ -1079,17 +1045,17 @@ FAssetHeaderPatcher::EResult FAssetHeaderPatcherInner::PatchHeader_WriteDestinat
 				}
 			}
 
-			const int64 AsetRegistryDataSize = Writer.Tell() - AsetRegistryDataStartOffset;
-			const int64 Delta = AsetRegistryDataSize - SourceSection.Size;
-			PatchSummaryOffsets(Summary, AsetRegistryDataStartOffset, Delta);
+			const int64 AssetRegistryDataSize = Writer.Tell() - AssetRegistryDataStartOffset;
+			const int64 Delta = AssetRegistryDataSize - SourceSection.Size;
+			PatchSummaryOffsets(Summary, AssetRegistryDataStartOffset, Delta);
 			Summary.TotalHeaderSize += (int32)Delta;
 
-			if (AsetRegistryData.PkgData.DependencyDataOffset != INDEX_NONE)
+			if (AssetRegistryData.PkgData.DependencyDataOffset != INDEX_NONE)
 			{
 				// DependencyDataOffset is not relative but points to just after the rest of the AR data
 				// We will seek back and write this later
-				const int64 DependencyDataDelta = AsetRegistryDataStartOffset - SourceSection.Offset + Delta;
-				AsetRegistryData.PkgData.DependencyDataOffset += DependencyDataDelta;
+				const int64 DependencyDataDelta = AssetRegistryDataStartOffset - SourceSection.Offset + Delta;
+				AssetRegistryData.PkgData.DependencyDataOffset += DependencyDataDelta;
 			}
 
 			break;
@@ -1141,11 +1107,11 @@ FAssetHeaderPatcher::EResult FAssetHeaderPatcherInner::PatchHeader_WriteDestinat
 		return EResult::ErrorFailedToWriteToDestinationFile;
 	}
 
-	if (AsetRegistryData.PkgData.DependencyDataOffset != INDEX_NONE)
+	if (AssetRegistryData.PkgData.DependencyDataOffset != INDEX_NONE)
 	{
 		// Re-write asset registry dependency data offset
 		Writer.Seek(Summary.AssetRegistryDataOffset);
-		Writer << AsetRegistryData.PkgData.DependencyDataOffset;
+		Writer << AssetRegistryData.PkgData.DependencyDataOffset;
 
 		if (Writer.IsError())
 		{
diff --git a/Engine/Source/Runtime/Engine/Private/WorldPartition/WorldPartitionActorDesc.cpp b/Engine/Source/Runtime/Engine/Private/WorldPartition/WorldPartitionActorDesc.cpp
--- a/Engine/Source/Runtime/Engine/Private/WorldPartition/WorldPartitionActorDesc.cpp
+++ b/Engine/Source/Runtime/Engine/Private/WorldPartition/WorldPartitionActorDesc.cpp
@@ -239,6 +239,8 @@ void FWorldPartitionActorDesc::Init(const FWorldPartitionActorDescInitData& Desc
 	
 	// Serialize metadata payload
 	FActorDescArchive ActorDescAr(MetadataAr, this);
+	ActorDescAr.Init();
+
 	Serialize(ActorDescAr);
 
 	// Call registered deprecator
@@ -256,6 +258,37 @@ void FWorldPartitionActorDesc::Init(const FWorldPartitionActorDescInitData& Desc
 	Container = nullptr;
 }
 
+void FWorldPartitionActorDesc::Patch(const FWorldPartitionActorDescInitData& DescData, TArray<uint8>& OutData, FWorldPartitionAssetDataPatcher* InAssetDataPatcher)
+{
+	// Serialize actor metadata
+	FMemoryReader MetadataAr(DescData.SerializedData, true);
+
+	// Serialize metadata custom versions
+	FCustomVersionContainer CustomVersions;
+	CustomVersions.Serialize(MetadataAr);
+	MetadataAr.SetCustomVersions(CustomVersions);
+	
+	// Patch metadata payload
+	TArray<uint8> PatchedPayloadData;
+	FMemoryWriter PatchedPayloadAr(PatchedPayloadData, true);
+
+	TUniquePtr<FWorldPartitionActorDesc> ActorDesc(AActor::StaticCreateClassActorDesc(DescData.NativeClass ? DescData.NativeClass : AActor::StaticClass()));
+	FActorDescArchivePatcher ActorDescAr(MetadataAr, ActorDesc.Get(), PatchedPayloadAr, InAssetDataPatcher);	
+	FTopLevelAssetPath ActorClassPath(TEXT("/Script/Engine.Actor"));
+	ActorDescAr.Init(ActorClassPath);
+
+	ActorDesc->Serialize(ActorDescAr);
+
+	// Serialize custom versions
+	TArray<uint8> HeaderData;
+	FMemoryWriter HeaderAr(HeaderData);
+	CustomVersions.Serialize(HeaderAr); 
+
+	// Append data
+	OutData = MoveTemp(HeaderData);
+	OutData.Append(PatchedPayloadData);
+}
+
 bool FWorldPartitionActorDesc::Equals(const FWorldPartitionActorDesc* Other) const
 {
 	return
@@ -337,6 +370,8 @@ void FWorldPartitionActorDesc::SerializeTo(TArray<uint8>& OutData)
 	TArray<uint8> PayloadData;
 	FMemoryWriter PayloadAr(PayloadData, true);
 	FActorDescArchive ActorDescAr(PayloadAr, this);
+	ActorDescAr.Init();
+
 	Serialize(ActorDescAr);
 
 	// Serialize custom versions
diff --git a/Engine/Source/Runtime/Engine/Private/WorldPartition/WorldPartitionActorDescArchive.cpp b/Engine/Source/Runtime/Engine/Private/WorldPartition/WorldPartitionActorDescArchive.cpp
--- a/Engine/Source/Runtime/Engine/Private/WorldPartition/WorldPartitionActorDescArchive.cpp
+++ b/Engine/Source/Runtime/Engine/Private/WorldPartition/WorldPartitionActorDescArchive.cpp
@@ -23,14 +23,17 @@ FActorDescArchive::FActorDescArchive(FArchive& InArchive, FWorldPartitionActorDe
 
 	SetIsPersistent(true);
 	SetIsLoading(InArchive.IsLoading());
+}
 
+void FActorDescArchive::Init(const FTopLevelAssetPath InClassPath)
+{
 	UsingCustomVersion(FUE5MainStreamObjectVersion::GUID);
 	UsingCustomVersion(FFortniteSeasonBranchObjectVersion::GUID);
 	UsingCustomVersion(FFortniteMainBranchObjectVersion::GUID);
 
 	if (CustomVer(FFortniteMainBranchObjectVersion::GUID) >= FFortniteMainBranchObjectVersion::WorldPartitionActorClassDescSerialize)
 	{
-		InnerArchive << InActorDesc->bIsDefaultActorDesc;
+		*this << ActorDesc->bIsDefaultActorDesc;
 	}
 
 	if (CustomVer(FFortniteSeasonBranchObjectVersion::GUID) >= FFortniteSeasonBranchObjectVersion::WorldPartitionActorDescNativeBaseClassSerialization)
@@ -38,28 +41,28 @@ FActorDescArchive::FActorDescArchive(FArchive& InArchive, FWorldPartitionActorDe
 		if (CustomVer(FUE5MainStreamObjectVersion::GUID) < FUE5MainStreamObjectVersion::WorldPartitionActorDescActorAndClassPaths)
 		{
 			FName BaseClassPathName;
-			InnerArchive << BaseClassPathName;
+			*this << BaseClassPathName;
 PRAGMA_DISABLE_DEPRECATION_WARNINGS
-			InActorDesc->BaseClass = FAssetData::TryConvertShortClassNameToPathName(BaseClassPathName);
+			ActorDesc->BaseClass = FAssetData::TryConvertShortClassNameToPathName(BaseClassPathName);
 PRAGMA_ENABLE_DEPRECATION_WARNINGS
 		}
 		else
 		{
-			InnerArchive << InActorDesc->BaseClass;
+			*this << ActorDesc->BaseClass;
 		}
 	}
 
 	if (CustomVer(FUE5MainStreamObjectVersion::GUID) < FUE5MainStreamObjectVersion::WorldPartitionActorDescActorAndClassPaths)
 	{
 		FName NativeClassPathName;
-		InnerArchive << NativeClassPathName;
+		*this << NativeClassPathName;
 PRAGMA_DISABLE_DEPRECATION_WARNINGS
-		InActorDesc->NativeClass = FAssetData::TryConvertShortClassNameToPathName(NativeClassPathName);
+		ActorDesc->NativeClass = FAssetData::TryConvertShortClassNameToPathName(NativeClassPathName);
 PRAGMA_ENABLE_DEPRECATION_WARNINGS
 	}
 	else
 	{
-		InnerArchive << InActorDesc->NativeClass;
+		*this << ActorDesc->NativeClass;
 	}
 
 	if (IsLoading())
@@ -84,14 +87,15 @@ PRAGMA_ENABLE_DEPRECATION_WARNINGS
 			}
 		};
 
-		TryRedirectClass(InActorDesc->NativeClass);
-		TryRedirectClass(InActorDesc->BaseClass);
+		TryRedirectClass(ActorDesc->NativeClass);
+		TryRedirectClass(ActorDesc->BaseClass);
 	}
 
 	// Get the class descriptor to do delta serialization
 	FWorldPartitionClassDescRegistry& ClassDescRegistry = FWorldPartitionClassDescRegistry::Get();
-	const FTopLevelAssetPath ClassPath = InActorDesc->BaseClass.IsValid() ? InActorDesc->BaseClass : InActorDesc->NativeClass;
-	ClassDesc = InActorDesc->bIsDefaultActorDesc ? ClassDescRegistry.GetClassDescDefaultForClass(ClassPath) : ClassDescRegistry.GetClassDescDefaultForActor(ClassPath);
+	const FTopLevelAssetPath ClassPath = InClassPath.IsValid() ? InClassPath : (ActorDesc->BaseClass.IsValid() ? ActorDesc->BaseClass : ActorDesc->NativeClass);
+	ClassDesc = (ActorDesc->bIsDefaultActorDesc && !InClassPath.IsValid()) ? ClassDescRegistry.GetClassDescDefaultForClass(ClassPath) : ClassDescRegistry.GetClassDescDefaultForActor(ClassPath);
+
 	if (!ClassDesc)
 	{
 		if (IsLoading())
@@ -101,11 +105,11 @@ PRAGMA_ENABLE_DEPRECATION_WARNINGS
 			ClassDesc = ClassDescRegistry.GetClassDescDefault(FTopLevelAssetPath(TEXT("/Script/Engine.Actor")));
 			check(ClassDesc);			
 
-			UE_LOG(LogWorldPartition, Log, TEXT("Can't find class descriptor '%s' for loading '%s', using '%s'"), *ClassPath.ToString(), *InActorDesc->GetActorSoftPath().ToString(), *ClassDesc->GetActorSoftPath().ToString());
+			UE_LOG(LogWorldPartition, Log, TEXT("Can't find class descriptor '%s' for loading '%s', using '%s'"), *ClassPath.ToString(), *ActorDesc->GetActorSoftPath().ToString(), *ClassDesc->GetActorSoftPath().ToString());
 		}
 		else
 		{
-			UE_LOG(LogWorldPartition, Log, TEXT("Can't find class descriptor '%s' for saving '%s'"), *ClassPath.ToString(), *InActorDesc->GetActorSoftPath().ToString());
+			UE_LOG(LogWorldPartition, Log, TEXT("Can't find class descriptor '%s' for saving '%s'"), *ClassPath.ToString(), *ActorDesc->GetActorSoftPath().ToString());
 		}
 	}
 
@@ -123,4 +127,32 @@ FArchive& FActorDescArchive::operator<<(FSoftObjectPath& Value)
 
 	return *this;
 }
+
+FArchive& FActorDescArchivePatcher::operator<<(FName& Value)
+{
+	TGuardValue<bool> GuardIsPatching(bIsPatching, true);
+	FActorDescArchive::operator<<(Value);
+	AssetDataPatcher->DoPatch(Value);
+	OutAr << Value;
+	return *this;
+}
+
+FArchive& FActorDescArchivePatcher::operator<<(FSoftObjectPath& Value)
+{
+	TGuardValue<bool> GuardIsPatching(bIsPatching, true);
+	FActorDescArchive::operator<<(Value);
+	AssetDataPatcher->DoPatch(Value);
+	OutAr << Value;
+	return *this;
+}
+
+void FActorDescArchivePatcher::Serialize(void* V, int64 Length)
+{
+	FActorDescArchive::Serialize(V, Length);
+
+	if (!bIsPatching)
+	{
+		OutAr.Serialize(V, Length);
+	}
+}
 #endif
diff --git a/Engine/Source/Runtime/Engine/Private/WorldPartition/WorldPartitionActorDescUtils.cpp b/Engine/Source/Runtime/Engine/Private/WorldPartition/WorldPartitionActorDescUtils.cpp
--- a/Engine/Source/Runtime/Engine/Private/WorldPartition/WorldPartitionActorDescUtils.cpp
+++ b/Engine/Source/Runtime/Engine/Private/WorldPartition/WorldPartitionActorDescUtils.cpp
@@ -145,6 +145,29 @@ FString FWorldPartitionActorDescUtils::GetAssetDataFromActorDescriptor(TUniquePt
 	return FBase64::Encode(SerializedData);
 }
 
+bool FWorldPartitionActorDescUtils::GetPatchedAssetDataFromAssetData(const FAssetData& InAssetData, FString& OutAssetData, FWorldPartitionAssetDataPatcher* InAssetDataPatcher)
+{
+	if (IsValidActorDescriptorFromAssetData(InAssetData))
+	{
+		FWorldPartitionActorDescInitData ActorDescInitData = FWorldPartitionActorDescInitData()
+			.SetNativeClass(GetActorNativeClassFromAssetData(InAssetData))
+			.SetPackageName(InAssetData.PackageName)
+			.SetActorPath(InAssetData.GetSoftObjectPath());
+
+		FString ActorMetaDataStr;
+		verify(InAssetData.GetTagValue(NAME_ActorMetaData, ActorMetaDataStr));
+		verify(FBase64::Decode(ActorMetaDataStr, ActorDescInitData.SerializedData));
+
+		TArray<uint8> PatchedData;
+		FWorldPartitionActorDesc::Patch(ActorDescInitData, PatchedData, InAssetDataPatcher);
+		OutAssetData = FBase64::Encode(PatchedData);
+
+		return OutAssetData != ActorMetaDataStr;
+	}
+
+	return false;
+}
+
 void FWorldPartitionActorDescUtils::UpdateActorDescriptorFromActor(const AActor* InActor, TUniquePtr<FWorldPartitionActorDesc>& OutActorDesc)
 {
 	TUniquePtr<FWorldPartitionActorDesc> NewActorDesc(InActor->CreateActorDesc());
diff --git a/Engine/Source/Runtime/Engine/Public/WorldPartition/WorldPartitionActorDesc.h b/Engine/Source/Runtime/Engine/Public/WorldPartition/WorldPartitionActorDesc.h
--- a/Engine/Source/Runtime/Engine/Public/WorldPartition/WorldPartitionActorDesc.h
+++ b/Engine/Source/Runtime/Engine/Public/WorldPartition/WorldPartitionActorDesc.h
@@ -28,6 +28,14 @@ struct FWorldPartitionActorDescInitData
 	FWorldPartitionActorDescInitData& SetActorPath(const FSoftObjectPath& InActorPath) { ActorPath = InActorPath; return *this; }
 };
 
+struct FWorldPartitionAssetDataPatcher
+{
+	virtual bool DoPatch(FString& InOutString) = 0;
+	virtual bool DoPatch(FName& InOutName) = 0;
+	virtual bool DoPatch(FSoftObjectPath& InOutSoft) = 0;
+	virtual bool DoPatch(FTopLevelAssetPath& InOutPath) = 0;
+};
+
 class AActor;
 class UActorDescContainer;
 class IStreamingGenerationErrorHandler;
@@ -255,7 +263,9 @@ public:
 	ENGINE_API virtual void Init(const AActor* InActor);
 	ENGINE_API virtual void Init(const FWorldPartitionActorDescInitData& DescData);
 
-	ENGINE_API virtual bool Equals(const FWorldPartitionActorDesc* Other) const;
+	static ENGINE_API void Patch(const FWorldPartitionActorDescInitData& DescData, TArray<uint8>& OutData, FWorldPartitionAssetDataPatcher* InAssetDataPatcher);
+
+	ENGINE_API virtual bool Equals(const FWorldPartitionActorDesc* Other) const;	
 
 	/**
 	 * Returns true if resaving this actor will have an impact on streaming generation. Before class descriptors, properties changed on a Blueprint 
diff --git a/Engine/Source/Runtime/Engine/Public/WorldPartition/WorldPartitionActorDescArchive.h b/Engine/Source/Runtime/Engine/Public/WorldPartition/WorldPartitionActorDescArchive.h
--- a/Engine/Source/Runtime/Engine/Public/WorldPartition/WorldPartitionActorDescArchive.h
+++ b/Engine/Source/Runtime/Engine/Public/WorldPartition/WorldPartitionActorDescArchive.h
@@ -4,16 +4,21 @@
 #if WITH_EDITOR
 #include "CoreMinimal.h"
 #include "Serialization/Archive.h"
+#include "UObject/TopLevelAssetPath.h"
 #include "UObject/FortniteMainBranchObjectVersion.h"
 
 class FWorldPartitionActorDesc;
+struct FWorldPartitionAssetDataPatcher;
 
 class FActorDescArchive : public FArchiveProxy
 {
 public:
 	FActorDescArchive(FArchive& InArchive, FWorldPartitionActorDesc* InActorDesc);
 
+	void Init(const FTopLevelAssetPath InClassPath = FTopLevelAssetPath());
+
 	//~ Begin FArchive Interface
+	virtual FArchive& operator<<(FName& Value) override { FArchiveProxy::operator<<(Value); return *this; }
 	virtual FArchive& operator<<(FText& Value) override { unimplemented(); return *this; }
 	virtual FArchive& operator<<(UObject*& Value) override { unimplemented(); return *this; }
 	virtual FArchive& operator<<(FLazyObjectPtr& Value) override { unimplemented(); return *this; }
@@ -117,6 +122,30 @@ public:
 	bool bIsMissingClassDesc;
 };
 
+class FActorDescArchivePatcher : public FActorDescArchive
+{
+public:
+	FActorDescArchivePatcher(FArchive& InArchive, FWorldPartitionActorDesc* InActorDesc, FArchive& OutArchive, FWorldPartitionAssetDataPatcher* InAssetDataPatcher)
+		: FActorDescArchive(InArchive, InActorDesc)
+		, OutAr(OutArchive)
+		, AssetDataPatcher(InAssetDataPatcher)
+		, bIsPatching(false)
+	{
+		check(AssetDataPatcher);
+	}
+
+	//~ Begin FArchive Interface
+	virtual FArchive& operator<<(FName& Value) override;
+	virtual FArchive& operator<<(FSoftObjectPath& Value) override;
+	virtual void Serialize(void* V, int64 Length) override;
+	//~ End FArchive Interface
+
+private:
+	FArchive& OutAr;
+	FWorldPartitionAssetDataPatcher* AssetDataPatcher;
+	bool bIsPatching;
+};
+
 template <typename DestPropertyType, typename SourcePropertyType = DestPropertyType>
 using TDeltaSerialize = FActorDescArchive::TDeltaSerializer<DestPropertyType, SourcePropertyType>;
 #endif
\ No newline at end of file
diff --git a/Engine/Source/Runtime/Engine/Public/WorldPartition/WorldPartitionActorDescUtils.h b/Engine/Source/Runtime/Engine/Public/WorldPartition/WorldPartitionActorDescUtils.h
--- a/Engine/Source/Runtime/Engine/Public/WorldPartition/WorldPartitionActorDescUtils.h
+++ b/Engine/Source/Runtime/Engine/Public/WorldPartition/WorldPartitionActorDescUtils.h
@@ -51,6 +51,15 @@ struct FWorldPartitionActorDescUtils
 	 */
 	static ENGINE_API FString GetAssetDataFromActorDescriptor(TUniquePtr<FWorldPartitionActorDesc>& InActorDesc);
 
+	/** 
+	 * Patches an actor descriptor's asset data
+	 * @param InAssetData			The asset data to be patched.
+	 * @param OutAssetData			The patched asset data.
+	 * @param InAssetDataPatcher	The patcher object.
+	 * @return						The actor descriptor patched data.
+	 */
+	static ENGINE_API bool GetPatchedAssetDataFromAssetData(const FAssetData& InAssetData, FString& OutAssetData, FWorldPartitionAssetDataPatcher* InAssetDataPatcher);
+
 	/** 
 	 * Update an actor descriptor with new values coming from the provided actor.
 	 * @param InActor		The actor to update from.
