AssetViewUtils::MoveAssets - Fix issue whe external packages weren't moved
* When moving an asset that has external packages (such as a map), make sure to mark all the external packages as needing to be saved

diff --git a/Engine/Source/Developer/AssetTools/Private/AssetRenameManager.cpp b/Engine/Source/Developer/AssetTools/Private/AssetRenameManager.cpp
--- a/Engine/Source/Developer/AssetTools/Private/AssetRenameManager.cpp
+++ b/Engine/Source/Developer/AssetTools/Private/AssetRenameManager.cpp
@@ -908,9 +908,27 @@ void FAssetRenameManager::LoadReferencingPackages(TArray<FAssetRenameDataWithRef
 		TArray<UPackage*> PackagesToSaveForThisAsset;
 		bool bAllPackagesLoadedForThisAsset = true;
 
+		TSet<FName> ReferencingExternalPackageNames;
+		if (Asset)
+		{
+			for (UPackage* ExternalPackage : Asset->GetPackage()->GetExternalPackages())
+			{
+				FName ExternalPackageName = ExternalPackage->GetFName();
+				ReferencingExternalPackageNames.Add(ExternalPackageName);
+				OutReferencingPackagesToSave.Add(ExternalPackage);
+			}
+		}
+
 		for (auto It = RenameData.ReferencingPackageNames.CreateIterator(); It; ++It)
 		{
 			FName PackageName = *It;
+
+			// Ignore external packages of this asset, those are already added to the list of packages to save
+			if (ReferencingExternalPackageNames.Contains(PackageName))
+			{
+				continue;
+			}
+
 			// Check if the package is a map before loading it!
 			if (!bLoadAllPackages && FEditorFileUtils::IsMapPackageAsset(PackageName.ToString()))
 			{
