Header patching.
Surfacing errors.
Removing many duplications of strings

diff --git a/Engine/Source/Developer/AssetTools/Private/AssetHeaderPatcher.cpp b/Engine/Source/Developer/AssetTools/Private/AssetHeaderPatcher.cpp
--- a/Engine/Source/Developer/AssetTools/Private/AssetHeaderPatcher.cpp
+++ b/Engine/Source/Developer/AssetTools/Private/AssetHeaderPatcher.cpp
@@ -253,11 +253,11 @@ class FAssetHeaderPatcherInner
 public:
 	using EResult = FAssetHeaderPatcher::EResult;
 
-	FAssetHeaderPatcherInner(FString InSrcAsset, FString InDstAsset, TMap<FString, FString> InSearchAndReplace, FArchive* InDstArchive = nullptr)
-		: SrcAsset(MoveTemp(InSrcAsset))
-		, DstAsset(MoveTemp(InDstAsset))
+	FAssetHeaderPatcherInner(const FString& InSrcAsset, const FString& InDstAsset, const TMap<FString, FString>& InSearchAndReplace, FArchive* InDstArchive = nullptr)
+		: SrcAsset(InSrcAsset)
+		, DstAsset(InDstAsset)
+		, SearchAndReplace(InSearchAndReplace)
 		, DstArchive(InDstArchive)
-		, SearchAndReplace(MoveTemp(InSearchAndReplace))
 	{
 	}
 
@@ -271,11 +271,12 @@ public:
 	void PatchHeader_PatchSections();
 	FAssetHeaderPatcher::EResult PatchHeader_WriteDestinationFile();
 
-	FString SrcAsset;
-	FString DstAsset;
+	const FString& SrcAsset;
+	const FString& DstAsset;
+	const TMap<FString, FString>& SearchAndReplace;
 	FArchive* DstArchive = nullptr;
 	TUniquePtr<FArchive> DstArchiveOwner;
-	TMap<FString, FString> SearchAndReplace;
+
 	TArray64<uint8> SrcBuffer;
 
 	struct FHeaderInformation
@@ -324,60 +325,26 @@ public:
 	FAssetRegistryData AssetRegistryData;
 };
 
-
-UE::Tasks::TTask<FAssetHeaderPatcher::EResult> FAssetHeaderPatcher::Start(FString InSrcAsset, FString InDstAsset,
-	TMap<FString, FString> InSearchAndReplace)
+FAssetHeaderPatcher::EResult FAssetHeaderPatcher::DoPatch(const FString& InSrcAsset, const FString& InDstAsset, const TMap<FString, FString>& InSearchAndReplace)
 {
-	TUniquePtr<FAssetHeaderPatcherInner> Self = MakeUnique<FAssetHeaderPatcherInner>(MoveTemp(InSrcAsset), MoveTemp(InDstAsset), MoveTemp(InSearchAndReplace));
-
-	return UE::Tasks::Launch(UE_SOURCE_LOCATION, [Self = MoveTemp(Self)]() -> FAssetHeaderPatcher::EResult
-		{
-			if (!FFileHelper::LoadFileToArray(Self->SrcBuffer, *Self->SrcAsset))
-			{
-				UE_LOG(LogAssetHeaderPatcher, Error, TEXT("Failed to load %s"), *Self->SrcAsset);
-				return FAssetHeaderPatcherInner::EResult::ErrorFailedToLoadSourceAsset;
-			}
-			else
-			{
-				return Self->PatchHeader();
-			}
-		});
-}
+	FAssetHeaderPatcherInner Inner(InSrcAsset, InDstAsset, InSearchAndReplace);
 
-UE::Tasks::TTask<FAssetHeaderPatcher::EResult> FAssetHeaderPatcher::Start(TUniquePtr<FArchive> InReader, FString InDstAsset,
-	TMap<FString, FString> InSearchAndReplace)
-{
-	if (!InReader || InReader->IsError())
+	if (!FFileHelper::LoadFileToArray(Inner.SrcBuffer, *Inner.SrcAsset))
 	{
-		UE_LOG(LogAssetHeaderPatcher, Error, TEXT("Failed to load %s (archive is invalid)"), InReader ? *InReader->GetArchiveName() : TEXT("Unknown"));
-
-		return UE::Tasks::Launch(UE_SOURCE_LOCATION, []() -> EResult { return EResult::ErrorFailedToLoadSourceAsset; });
+		UE_LOG(LogAssetHeaderPatcher, Error, TEXT("Failed to load %s"), *Inner.SrcAsset);
+		return FAssetHeaderPatcherInner::EResult::ErrorFailedToLoadSourceAsset;
+	}
+	else
+	{
+		return Inner.PatchHeader();
 	}
-
-	TUniquePtr<FAssetHeaderPatcherInner> Self = MakeUnique<FAssetHeaderPatcherInner>(InReader->GetArchiveName(), MoveTemp(InDstAsset), MoveTemp(InSearchAndReplace));
-
-	return UE::Tasks::Launch(UE_SOURCE_LOCATION, [Self = MoveTemp(Self), Reader = MoveTemp(InReader)]() -> EResult
-		{
-			Reader->Seek(0);
-			Self->SrcBuffer.SetNumUninitialized(Reader->TotalSize());
-			Reader->Serialize(Self->SrcBuffer.GetData(), Self->SrcBuffer.Num());
-
-			if (Reader->IsError())
-			{
-				UE_LOG(LogAssetHeaderPatcher, Error, TEXT("Failed to load %s"), *Self->SrcAsset);
-				return FAssetHeaderPatcherInner::EResult::ErrorFailedToLoadSourceAsset;
-			}
-			else
-			{
-				return Self->PatchHeader();
-			}
-		});
 }
 
+
 FAssetHeaderPatcher::EResult FAssetHeaderPatcher::Test_DoPatch(FArchive& InReader, FArchive& InWriter,
-	TMap<FString, FString> InSearchAndReplace)
+	const TMap<FString, FString>& InSearchAndReplace)
 {
-	FAssetHeaderPatcherInner Inner(InReader.GetArchiveName(), InWriter.GetArchiveName(), MoveTemp(InSearchAndReplace), &InWriter);
+	FAssetHeaderPatcherInner Inner(InReader.GetArchiveName(), InWriter.GetArchiveName(), InSearchAndReplace, &InWriter);
 
 	InReader.Seek(0);
 	Inner.SrcBuffer.SetNumUninitialized(InReader.TotalSize());
diff --git a/Engine/Source/Developer/AssetTools/Private/AssetHeaderPatcher.h b/Engine/Source/Developer/AssetTools/Private/AssetHeaderPatcher.h
--- a/Engine/Source/Developer/AssetTools/Private/AssetHeaderPatcher.h
+++ b/Engine/Source/Developer/AssetTools/Private/AssetHeaderPatcher.h
@@ -20,12 +20,7 @@ struct FAssetHeaderPatcher
 		ErrorEmptyRequireSection,
 	};
 
-	static UE::Tasks::TTask<EResult> Start(FString InSrcAsset, FString InDstAsset,
-		TMap<FString, FString> InSearchAndReplace);
+	static EResult DoPatch(const FString& InSrcAsset, const FString& InDstAsset, const TMap<FString, FString>& InSearchAndReplace);
 
-	static UE::Tasks::TTask<EResult> Start(TUniquePtr<FArchive> InSrcReader, FString InDstAsset,
-		TMap<FString, FString> InSearchAndReplace);
-
-	static EResult Test_DoPatch(FArchive& InSrcReader, FArchive& InDstWriter,
-		TMap<FString, FString> InSearchAndReplace);
+	static EResult Test_DoPatch(FArchive& InSrcReader, FArchive& InDstWriter, const TMap<FString, FString>& InSearchAndReplace);
 };
diff --git a/Engine/Source/Developer/AssetTools/Private/AssetTools.cpp b/Engine/Source/Developer/AssetTools/Private/AssetTools.cpp
--- a/Engine/Source/Developer/AssetTools/Private/AssetTools.cpp
+++ b/Engine/Source/Developer/AssetTools/Private/AssetTools.cpp
@@ -10,6 +10,7 @@
 #include "Misc/Paths.h"
 #include "Misc/PathViews.h"
 #include "Misc/ScopedSlowTask.h"
+#include "Misc/ScopeLock.h"
 #include "UObject/GCObjectScopeGuard.h"
 #include "UObject/UObjectHash.h"
 #include "UObject/UObjectIterator.h"
@@ -2291,17 +2292,15 @@ bool UAssetToolsImpl::AdvancedCopyPackages(
 			std::atomic<int32> PatchAssetsCompletedCount = 0;
 			UE::Tasks::FTaskEvent PatchAssetsCompletionTask{ UE_SOURCE_LOCATION };
 
-			TMap<FString, FString> ToCopyAndPatch = AllSourceAndDestPackages(SourceAndDestPackages);
+			TMap<FString, FString> ToCopyAndPatchPackages = AllSourceAndDestPackages(SourceAndDestPackages);
 			TMap<FString, FString> PatchingPatterns = GenerateAdditionalAssetMappings(SourceAndDestPackages);
-			PatchingPatterns.Append(ToCopyAndPatch);
+			PatchingPatterns.Append(ToCopyAndPatchPackages);
 
-			LoopProgress = MakeUnique<FScopedSlowTask>(static_cast<float>(ToCopyAndPatch.Num()), LOCTEXT("AdvancedCopyPackages.ReplacingAssetReferences", "Replacing Asset References..."));
-			LoopProgress->MakeDialog();
-
-			for (const TTuple<FString, FString>& Package : ToCopyAndPatch)
+			// Construct all filenames
+			TMap<FString, FString> ToCopyAndPatchFiles;
+			ToCopyAndPatchFiles.Reserve(ToCopyAndPatchPackages.Num());
+			for (const TTuple<FString, FString>& Package : ToCopyAndPatchPackages)
 			{
-				LoopProgress->EnterProgressFrame();
-
 				const FString& PackageName = Package.Key;
 				const FString& DestPackage = Package.Value;
 				FString SrcFilename;
@@ -2309,16 +2308,37 @@ bool UAssetToolsImpl::AdvancedCopyPackages(
 				if (FPackageName::DoesPackageExist(PackageName, &SrcFilename))
 				{
 					FString DestFilename = FPackageName::LongPackageNameToFilename(DestPackage, FString(FPathViews::GetExtension(SrcFilename, true)));
-					UE::Tasks::TTask<FAssetHeaderPatcher::EResult> PatcherTask = FAssetHeaderPatcher::Start(SrcFilename, DestFilename, PatchingPatterns);
-					UE::Tasks::FTask PatcherLifetimeTask = UE::Tasks::Launch(UE_SOURCE_LOCATION, [&PatchAssetsCompletedCount]
-						{
-							PatchAssetsCompletedCount.fetch_add(1, std::memory_order_relaxed);
-						},
-						MoveTemp(PatcherTask), UE::Tasks::ETaskPriority::Default, UE::Tasks::EExtendedTaskPriority::Inline);
-					PatchAssetsCompletionTask.AddPrerequisites(PatcherLifetimeTask);
+					ToCopyAndPatchFiles.Add({ MoveTemp(SrcFilename), MoveTemp(DestFilename) });
 				}
 			}
 
+			LoopProgress = MakeUnique<FScopedSlowTask>(static_cast<float>(ToCopyAndPatchFiles.Num()), LOCTEXT("AdvancedCopyPackages.ReplacingAssetReferences", "Replacing Asset References..."));
+			LoopProgress->MakeDialog();
+
+			TSet<FString> ErroredFiles;
+			FCriticalSection  ErroredFilesLock;
+
+			// Spawn tasks (Scatter)
+			for (const TTuple<FString, FString>& Filename : ToCopyAndPatchFiles)
+			{
+				const FString& SrcFilename = Filename.Key;
+				const FString& DestFilename = Filename.Value;
+
+				UE::Tasks::FTask PatcherTask = UE::Tasks::Launch(UE_SOURCE_LOCATION,
+					[&PatchAssetsCompletedCount, &PatchingPatterns, InSrcFilename = SrcFilename, InDestFilename = DestFilename, &ErroredFilesLock, &ErroredFiles] () 
+					{
+						FAssetHeaderPatcher::EResult Result = FAssetHeaderPatcher::DoPatch(InSrcFilename, InDestFilename, PatchingPatterns);
+						if (Result != FAssetHeaderPatcher::EResult::Success) 
+						{
+							FScopeLock Lock(&ErroredFilesLock);
+							ErroredFiles.Add(InSrcFilename);
+						}
+						PatchAssetsCompletedCount.fetch_add(1, std::memory_order_relaxed);
+					});
+				PatchAssetsCompletionTask.AddPrerequisites(PatcherTask);
+			}
+
+			// Gather
 			PatchAssetsCompletionTask.Trigger();
 			while (!PatchAssetsCompletionTask.Wait(FTimespan::FromSeconds(0.5)))
 			{
@@ -2326,11 +2346,86 @@ bool UAssetToolsImpl::AdvancedCopyPackages(
 				LoopProgress->TickProgress();
 			}
 
-			for (const TTuple<FString, FString>& Package : SourceAndDestPackages)
+			// Sorting out the results.
+			// And reporting to the user
+			FMessageLog AdvancedCopyLog("AssetTools");
+
+			bool bHasErrors = ErroredFiles.Num() != 0;
+
+			if (bHasErrors)
+			{
+				AdvancedCopyLog.NewPage(LOCTEXT("AdvancedCopyPackages_SourceControlErrorsListPage", "Revision Control Errors"));
+			}
+
+			// reporting files with copy errors and filtering successful ones
+			for (const TTuple<FString, FString>& Filename : ToCopyAndPatchFiles)
 			{			
-				SuccessfullyCopiedSourcePackages.Add(FName(Package.Key));
-				SuccessfullyCopiedDestinationFiles.Add(Package.Value);
+				if (ErroredFiles.Contains(Filename.Key))
+				{
+					AdvancedCopyLog.Error(FText::Format(LOCTEXT("AdvancedCopyPackages_SourceControlError", "{0} could not be processed"), FText::FromString(*Filename.Key)));
+					continue;
+				}
+				SuccessfullyCopiedDestinationFiles.Add(Filename.Value);
 			}
+
+			if (SuccessfullyCopiedDestinationFiles.Num() > 0
+			&& GetDefault<UEditorLoadingSavingSettings>()->bSCCAutoAddNewFiles
+			&& ISourceControlModule::Get().IsEnabled())
+			{
+				// attempt to add files to source control (this can quite easily fail, but if it works it is very useful)
+				ISourceControlProvider& SourceControlProvider = ISourceControlModule::Get().GetProvider();
+				if (SourceControlProvider.Execute(ISourceControlOperation::Create<FMarkForAdd>(), SourceControlHelpers::AbsoluteFilenames(SuccessfullyCopiedDestinationFiles)) == ECommandResult::Failed)
+				{
+					for (const FString& Filename : SuccessfullyCopiedDestinationFiles)
+					{
+						if (!SourceControlProvider.GetState(*Filename, EStateCacheUsage::Use)->IsAdded())
+						{
+							if (!bHasErrors)
+							{
+								bHasErrors = true;
+								AdvancedCopyLog.NewPage(LOCTEXT("AdvancedCopyPackages_SourceControlErrorsListPage", "Revision Control Errors"));
+							}
+
+							AdvancedCopyLog.Error(FText::Format(LOCTEXT("AdvancedCopyPackages_SourceControlError", "{0} could not be added to revision control"), FText::FromString(*Filename)));
+						}
+					}
+				}
+			}
+
+			// Report the result to the user
+			FText LogMessage = FText::FromString(TEXT("Advanced content copy completed successfully!"));
+			EMessageSeverity::Type Severity = EMessageSeverity::Info;
+			FString ErrorsString;
+			if (bHasErrors)
+			{
+				Severity = EMessageSeverity::Error;
+
+				FString ErrorMessage = LOCTEXT("AdvancedCopyPackages_SourceControlErrorsList", "Some files reported revision control errors.").ToString();
+
+				if (SuccessfullyCopiedSourcePackages.Num() > 0)
+				{
+					AdvancedCopyLog.NewPage(LOCTEXT("AdvancedCopyPackages_CopyErrorsSuccesslistPage", "Copied Successfully"));
+					for (const FString& Filename : SuccessfullyCopiedDestinationFiles)
+					{
+						AdvancedCopyLog.Info(FText::FromString(Filename));
+					}
+
+					ErrorMessage += LINE_TERMINATOR;
+					ErrorMessage += LOCTEXT("AdvancedCopyPackages_CopyErrorsSuccesslist", "Some files were copied successfully.").ToString();
+				}
+				LogMessage = FText::FromString(ErrorMessage);
+			}
+			else
+			{
+				AdvancedCopyLog.NewPage(LOCTEXT("AdvancedCopyPackages_CompletePage", "Advanced content copy completed successfully!"));
+				for (const FString& Filename : SuccessfullyCopiedDestinationFiles)
+				{
+					AdvancedCopyLog.Info(FText::FromString(Filename));
+				}
+			}
+			// @note Using the bForce param because the InSeverityFilter param is only checked against logs in the last page
+			AdvancedCopyLog.Notify(LogMessage, /*InSeverityFilter=*/EMessageSeverity::Error, /*bForce=*/(Severity <= NotificationSeverityFilter));
+			return true;
 		}
 		else
 		{
@@ -2445,22 +2540,19 @@ bool UAssetToolsImpl::AdvancedCopyPackages(
 				}
 			}
 
-			if (!UE::AssetTools::Private::bEnableHeaderPatching)
+			for (FName Dependency : Dependencies)
 			{
-				for (FName Dependency : Dependencies)
+				const int32 DependencyIndex = SuccessfullyCopiedSourcePackages.IndexOfByKey(Dependency);
+				if (DependencyIndex != INDEX_NONE)
 				{
-					const int32 DependencyIndex = SuccessfullyCopiedSourcePackages.IndexOfByKey(Dependency);
-					if (DependencyIndex != INDEX_NONE)
+					Consolidations.Reserve(Consolidations.Num() + DuplicatedObjectsForEachPackage[DependencyIndex].Num());
+					for (const TPair<TSoftObjectPtr<UObject>, TSoftObjectPtr<UObject>>& Duplication : DuplicatedObjectsForEachPackage[DependencyIndex])
 					{
-						Consolidations.Reserve(Consolidations.Num() + DuplicatedObjectsForEachPackage[DependencyIndex].Num());
-						for (const TPair<TSoftObjectPtr<UObject>, TSoftObjectPtr<UObject>>& Duplication : DuplicatedObjectsForEachPackage[DependencyIndex])
+						UObject* SourceObject = Duplication.Key.Get();
+						UObject* NewObject = Duplication.Value.Get();
+						if (SourceObject && NewObject)
 						{
-							UObject* SourceObject = Duplication.Key.Get();
-							UObject* NewObject = Duplication.Value.Get();
-							if (SourceObject && NewObject)
-							{
-								Consolidations.FindOrAdd(NewObject).AddUnique(SourceObject);
-							}
+							Consolidations.FindOrAdd(NewObject).AddUnique(SourceObject);
 						}
 					}
 				}
