Fix for duplicate LOCTEXT ID

diff --git a/Engine/Source/Developer/AssetTools/Private/AssetTools.cpp b/Engine/Source/Developer/AssetTools/Private/AssetTools.cpp
--- a/Engine/Source/Developer/AssetTools/Private/AssetTools.cpp
+++ b/Engine/Source/Developer/AssetTools/Private/AssetTools.cpp
@@ -2362,7 +2362,7 @@ bool UAssetToolsImpl::AdvancedCopyPackages(
 			{			
 				if (ErroredFiles.Contains(Filename.Key))
 				{
-					AdvancedCopyLog.Error(FText::Format(LOCTEXT("AdvancedCopyPackages_SourceControlError", "{0} could not be processed"), FText::FromString(*Filename.Key)));
+					AdvancedCopyLog.Error(FText::Format(LOCTEXT("AdvancedCopyPackages_CouldNotProcessError", "{0} could not be processed"), FText::FromString(*Filename.Key)));
 					continue;
 				}
 				SuccessfullyCopiedDestinationFiles.Add(Filename.Value);
