Automatically make an asset public when moving it in another plugin if it's going to be referenced from another plugin after the move (including redirectors)
Don't move an asset if it cannot be made public or if it ends up referencing a private asset that's not being moved or if a dependent asset cannot be moved to a different plugin
Only enabled if AssetTools.EnablePublicAssetFeature cvar is set (which replaces ContentBrowser.PublicAsset.EnablePublicAssetFeature)
#rnx

diff --git a/Engine/Source/Developer/AssetTools/Private/AssetRenameManager.cpp b/Engine/Source/Developer/AssetTools/Private/AssetRenameManager.cpp
--- a/Engine/Source/Developer/AssetTools/Private/AssetRenameManager.cpp
+++ b/Engine/Source/Developer/AssetTools/Private/AssetRenameManager.cpp
@@ -5,6 +5,7 @@
 #include "Serialization/ArchiveUObject.h"
 #include "UObject/Class.h"
 #include "Misc/PackageName.h"
+#include "Misc/PathViews.h"
 #include "Misc/MessageDialog.h"
 #include "HAL/FileManager.h"
 #include "Misc/FeedbackContext.h"
@@ -91,7 +92,8 @@ namespace AssetRenameManagerImpl
 
 struct FAssetRenameDataWithReferencers : public FAssetRenameData
 {
-	TSet<FName> ReferencingPackageNames;
+	TSet<FName> NotRenamedReferencingPackageNames;
+	TMap<FName, FName> RenamedReferencingPackageNames;
 	FText FailureReason;
 	bool bCreateRedirector;
 	bool bRenameFailed;
@@ -427,6 +429,12 @@ bool FAssetRenameManager::FixReferencesAndRename(const TArray<FAssetRenameData>&
 			// If any referencing packages are left read-only, the checkout failed or SCC was not enabled. Trim them from the save list and leave redirectors.
 			DetectReadOnlyPackages(AssetsToRename, ReferencingPackagesToSave);
 
+
+			// Make public any asset that will be referenced from another plugin after the rename.
+			// If the asset cannot be made public or if moving it requires a referenced asset that's not being modified to become public,
+			// its rename will fail and other assets that have dependencies between them will also fail to be renamed.
+			SetupPublicAssets(AssetsToRename);
+
 			if (bSoftReferencesOnly)
 			{
 				if (ReferencingPackagesToSave.Num() > 0)
@@ -736,7 +744,7 @@ void FAssetRenameManager::FindCDOReferences(const TArrayView<FAssetRenameDataWit
 void FAssetRenameManager::PopulateAssetReferencers(TArray<FAssetRenameDataWithReferencers>& AssetsToPopulate) const
 {
 	FAssetRegistryModule& AssetRegistryModule = FModuleManager::LoadModuleChecked<FAssetRegistryModule>("AssetRegistry");
-	TSet<FName> RenamingAssetPackageNames;
+	TMap<FName, FName> PackageNamesToRename;
 
 	// Get the names of all the packages containing the assets we are renaming so they arent added to the referencing packages list
 	for (FAssetRenameDataWithReferencers& AssetToRename : AssetsToPopulate)
@@ -744,7 +752,7 @@ void FAssetRenameManager::PopulateAssetReferencers(TArray<FAssetRenameDataWithRe
 		// If we're only fixing soft references we want to check for references inside the original package as we don't save the original package automatically
 		if (!AssetToRename.bOnlyFixSoftReferences)
 		{
-			RenamingAssetPackageNames.Add(FName(*AssetToRename.OldObjectPath.GetLongPackageName()));
+			PackageNamesToRename.Add(AssetToRename.OldObjectPath.GetLongPackageFName(), AssetToRename.NewObjectPath.GetLongPackageFName());
 		}
 	}
 
@@ -758,9 +766,10 @@ void FAssetRenameManager::PopulateAssetReferencers(TArray<FAssetRenameDataWithRe
 	// Gather all referencing packages for all assets that are being renamed
 	for (FAssetRenameDataWithReferencers& AssetToRename : AssetsToPopulate)
 	{
-		AssetToRename.ReferencingPackageNames.Empty();
+		AssetToRename.NotRenamedReferencingPackageNames.Empty();
+		AssetToRename.RenamedReferencingPackageNames.Empty();
 
-		FName OldPackageName = FName(*AssetToRename.OldObjectPath.GetLongPackageName());
+		FName OldPackageName = AssetToRename.OldObjectPath.GetLongPackageFName();
 
 		TMap<FName, TArray<FName>>& ReferencersMap = AssetToRename.bOnlyFixSoftReferences ? SoftReferencers : PackageReferencers;
 		if (!ReferencersMap.Contains(OldPackageName))
@@ -771,21 +780,25 @@ void FAssetRenameManager::PopulateAssetReferencers(TArray<FAssetRenameDataWithRe
 
 		for (const FName& ReferencingPackageName : ReferencersMap.FindChecked(OldPackageName))
 		{
-			if (!RenamingAssetPackageNames.Contains(ReferencingPackageName))
+			if (FName* NewPackageName = PackageNamesToRename.Find(ReferencingPackageName))
+			{
+				AssetToRename.RenamedReferencingPackageNames.Add(ReferencingPackageName, *NewPackageName);
+			}
+			else
 			{
-				AssetToRename.ReferencingPackageNames.Add(ReferencingPackageName);
+				AssetToRename.NotRenamedReferencingPackageNames.Add(ReferencingPackageName);
 			}
 		}
 
 		if (AssetToRename.bOnlyFixSoftReferences)
 		{
-			AssetToRename.ReferencingPackageNames.Add(FName(*AssetToRename.OldObjectPath.GetLongPackageName()));
-			AssetToRename.ReferencingPackageNames.Add(FName(*AssetToRename.NewObjectPath.GetLongPackageName()));
+			AssetToRename.NotRenamedReferencingPackageNames.Add(OldPackageName);
+			AssetToRename.NotRenamedReferencingPackageNames.Add(AssetToRename.NewObjectPath.GetLongPackageFName());
 
 			// Add dirty packages and the package that owns the reference. They will get filtered out in LoadReferencingPackages if they aren't valid
 			for (UPackage* Package : ExtraPackagesToCheckForSoftReferences)
 			{
-				AssetToRename.ReferencingPackageNames.Add(Package->GetFName());
+				AssetToRename.NotRenamedReferencingPackageNames.Add(Package->GetFName());
 			}
 		}
 	}
@@ -919,7 +932,7 @@ void FAssetRenameManager::LoadReferencingPackages(TArray<FAssetRenameDataWithRef
 			}
 		}
 
-		for (auto It = RenameData.ReferencingPackageNames.CreateIterator(); It; ++It)
+		for (auto It = RenameData.NotRenamedReferencingPackageNames.CreateIterator(); It; ++It)
 		{
 			FName PackageName = *It;
 
@@ -1012,7 +1025,7 @@ void FAssetRenameManager::GatherReferencingObjects(TArray<FAssetRenameDataWithRe
 			continue;
 		}
 
-		for (FName PackageName : RenameData.ReferencingPackageNames)
+		for (FName PackageName : RenameData.NotRenamedReferencingPackageNames)
 		{
 			UPackage* Package = FindPackage(nullptr, *PackageName.ToString());
 
@@ -1213,7 +1226,7 @@ void FAssetRenameManager::DetectReadOnlyPackages(TArray<FAssetRenameDataWithRefe
 					// Find all assets that were referenced by this package to create a redirector when named
 					for (FAssetRenameDataWithReferencers& RenameData : AssetsToRename)
 					{
-						if (RenameData.ReferencingPackageNames.Contains(PackageName))
+						if (RenameData.NotRenamedReferencingPackageNames.Contains(PackageName))
 						{
 							RenameData.bCreateRedirector = true;
 						}
@@ -1227,6 +1240,247 @@ void FAssetRenameManager::DetectReadOnlyPackages(TArray<FAssetRenameDataWithRefe
 	}
 }
 
+void FAssetRenameManager::SetupPublicAssets(TArray<FAssetRenameDataWithReferencers>& AssetsToRename) const
+{
+	static const IConsoleVariable* EnablePublicAssetFeatureCVar = IConsoleManager::Get().FindConsoleVariable(TEXT("AssetTools.EnablePublicAssetFeature"));
+	if (!EnablePublicAssetFeatureCVar || !EnablePublicAssetFeatureCVar->GetBool())
+	{
+		return;
+	}
+
+	FScopedSlowTask SlowTask((float)AssetsToRename.Num(), LOCTEXT("SetupPublicAssets", "Setting up public assets..."));
+	SlowTask.MakeDialog();
+
+	IAssetTools& AssetTools = IAssetTools::Get();
+	IAssetRegistry& AssetRegistry = IAssetRegistry::GetChecked();
+
+	// Build a map of old package name to rename data
+	TMap<FName, FAssetRenameDataWithReferencers*> OldPackageNameToRenameData;
+	for (FAssetRenameDataWithReferencers& RenameData : AssetsToRename)
+	{
+		OldPackageNameToRenameData.Add(RenameData.OldObjectPath.GetLongPackageFName(), &RenameData);
+	}
+
+	TArray<UPackage*> PackagesToMakePublic;
+	TSet<FName> RenamedReferencedPackageNames;
+	bool bSomeAssetCannotBeMovedToAnotherPlugin = false;
+	
+	// Determine which moved asset need to become public
+	for (FAssetRenameDataWithReferencers& RenameData : AssetsToRename)
+	{
+		SlowTask.EnterProgressFrame();
+
+		// Nothing to do if it won't be moved
+		if (RenameData.bRenameFailed || RenameData.bOnlyFixSoftReferences)
+		{
+			continue;
+		}
+
+		UObject* Asset = RenameData.Asset.Get();
+		UPackage* Package = Asset ? Asset->GetPackage() : nullptr;
+		if (!ensure(Package))
+		{
+			continue;
+		}
+
+		// Nothing to do if the renamed package mount point doesn't change
+		const FNameBuilder OldPackageName(RenameData.OldObjectPath.GetLongPackageFName());
+		const FNameBuilder NewPackageName(RenameData.NewObjectPath.GetLongPackageFName());
+		const FStringView OldPackageMountPoint = FPathViews::GetMountPointNameFromPath(OldPackageName);
+		const FStringView NewPackageMountPoint = FPathViews::GetMountPointNameFromPath(NewPackageName);
+		if (NewPackageMountPoint.Equals(OldPackageMountPoint, ESearchCase::IgnoreCase))
+		{
+			continue;
+		}
+
+		// Make sure the asset won't be referencing a private asset in a different mount point after the move
+		{
+			TArray<FName> DependencyPackageNames;
+			AssetRegistry.GetDependencies(Package->GetFName(), DependencyPackageNames);
+			for (FName DependencyPackageName : DependencyPackageNames)
+			{
+				if (OldPackageNameToRenameData.Find(DependencyPackageName))
+				{
+					RenamedReferencedPackageNames.Add(DependencyPackageName);
+
+					// Dependency is being renamed as well so it'll be handled later
+					continue;
+				}
+
+				{
+					const FNameBuilder DepPackageNameBuilder(DependencyPackageName);
+					const FStringView DependencyMountPoint = FPathViews::GetMountPointNameFromPath(DepPackageNameBuilder);
+					if (DependencyMountPoint.Equals(NewPackageMountPoint, ESearchCase::IgnoreCase))
+					{
+						// Dependency is under the same mount point as the new asset location
+						continue;
+					}
+				}
+
+				{
+					TArray<FAssetData> DependencyAssetDatas;
+					if (AssetRegistry.GetAssetsByPackageName(DependencyPackageName, DependencyAssetDatas) && !DependencyAssetDatas.IsEmpty())
+					{
+						if (DependencyAssetDatas[0].HasAnyPackageFlags(PKG_NotExternallyReferenceable))
+						{
+							RenameData.bRenameFailed = true;
+							RenameData.FailureReason = FText::Format(
+								LOCTEXT("MovedAssetReferencingPrivateAsset", "Cannot move asset to {0} because it would be referencing a private asset in a different plugin: {1}"),
+								FText::FromStringView(NewPackageMountPoint),
+								FText::FromName(DependencyPackageName));
+							break;
+						}
+					}
+				}
+			}
+
+			if (RenameData.bRenameFailed)
+			{
+				bSomeAssetCannotBeMovedToAnotherPlugin = true;
+				continue;
+			}
+		}
+
+		// If the asset is already public, it can be referenced from anywhere so we're good
+		if (Package->IsExternallyReferenceable())
+		{
+			continue;
+		}
+
+		// The asset doesn't have to become public if it's not referenced at all
+		if (!RenameData.bCreateRedirector && RenameData.NotRenamedReferencingPackageNames.IsEmpty() && RenameData.RenamedReferencingPackageNames.IsEmpty())
+		{
+			continue;
+		}
+
+		// Figure out if it's gonna be referenced from another mount point after the move
+		FString ReferencingAssetInDifferentMountPoint;
+		if (RenameData.bCreateRedirector)
+		{
+			ReferencingAssetInDifferentMountPoint = OldPackageName.ToString();
+		}
+		if (ReferencingAssetInDifferentMountPoint.IsEmpty())
+		{
+			for (FName It : RenameData.NotRenamedReferencingPackageNames)
+			{
+				const FNameBuilder ReferencingPackageName(It);
+				const FStringView ReferencingPackagMountPoint = FPathViews::GetMountPointNameFromPath(ReferencingPackageName);
+				if (!ReferencingPackagMountPoint.Equals(NewPackageMountPoint, ESearchCase::IgnoreCase))
+				{
+					ReferencingAssetInDifferentMountPoint = ReferencingPackageName.ToString();
+					break;
+				}
+			}
+		}
+		if (ReferencingAssetInDifferentMountPoint.IsEmpty())
+		{
+			for (const TPair<FName, FName>& It : RenameData.RenamedReferencingPackageNames)
+			{
+				const FName OldReferencingPackageName = It.Key;
+				const FName NewReferencingPackageName = It.Value;
+
+				FAssetRenameDataWithReferencers** FoundReferencingRenameData = OldPackageNameToRenameData.Find(OldReferencingPackageName);
+				if (!ensureAlways(FoundReferencingRenameData))
+				{
+					continue;
+				}
+				FAssetRenameDataWithReferencers& ReferencingRenameData = **FoundReferencingRenameData;
+
+				FNameBuilder ReferencingPackageName;
+				if (ReferencingRenameData.bRenameFailed || ReferencingRenameData.bOnlyFixSoftReferences)
+				{
+					OldReferencingPackageName.AppendString(ReferencingPackageName);
+				}
+				else
+				{
+					NewReferencingPackageName.AppendString(ReferencingPackageName);
+				}
+
+				const FStringView ReferencingPackagMountPoint = FPathViews::GetMountPointNameFromPath(ReferencingPackageName);
+				if (!NewPackageMountPoint.Equals(ReferencingPackagMountPoint, ESearchCase::IgnoreCase))
+				{
+					ReferencingAssetInDifferentMountPoint = ReferencingPackageName.ToString();
+					break;
+				}
+			}
+		}
+
+		// Check if the asset can be made public
+		if (!ReferencingAssetInDifferentMountPoint.IsEmpty())
+		{
+			const FString OldAssetPath = RenameData.OldObjectPath.GetAssetPathString();
+			if (AssetTools.CanAssetBePublic(OldAssetPath))
+			{
+				PackagesToMakePublic.Add(Package);
+			}
+			else
+			{
+				bSomeAssetCannotBeMovedToAnotherPlugin = true;
+				RenameData.bRenameFailed = true;
+				if (RenameData.bCreateRedirector)
+				{
+					RenameData.FailureReason = FText::Format(
+						LOCTEXT("AssetCannotBePublicForRedirector", "Cannot move asset to {0} because it cannot be made public in order to be referenced from a different plugin by its redirector"),
+						FText::FromStringView(NewPackageMountPoint));
+				}
+				else
+				{
+					RenameData.FailureReason = FText::Format(
+						LOCTEXT("AssetCannotBePublic", "Cannot move asset to {0} because it cannot be made public in order to be referenced from a different plugin by {1}"), 
+						FText::FromStringView(NewPackageMountPoint),
+						FText::FromString(ReferencingAssetInDifferentMountPoint));
+				}
+			}
+		}
+	}
+
+	if (bSomeAssetCannotBeMovedToAnotherPlugin)
+	{
+		// When an asset cannot be moved to another plugin, it changes the set of inter-plugin dependencies and it can require other assets 
+		// to become public as a result, which might not be intended by the user and those cascading effects would be difficult to grok.
+		// In this case, only rename assets that have no dependency against each other.
+		// @fixme: This is a little too aggressive. Ideally we'd check whether each asset actually has a direct or indirect
+		//         dependency with an asset that failed to be moved to another plugin but that could be expensive.
+		const FText FailureReason = LOCTEXT("DependentAssetCannotBeMoved", "Cannot rename asset because dependent assets could not be moved to a different plugin");
+		TSet<FName> RenamedReferencingPackageNames;
+
+		for (FAssetRenameDataWithReferencers& RenameData : AssetsToRename)
+		{
+			for (const TPair<FName, FName>& It : RenameData.RenamedReferencingPackageNames)
+			{
+				RenamedReferencingPackageNames.Add(It.Key);
+			}
+
+			if (!RenameData.bRenameFailed && !RenameData.bOnlyFixSoftReferences && !RenameData.RenamedReferencingPackageNames.IsEmpty())
+			{
+				RenameData.bRenameFailed = true;
+				RenameData.FailureReason = FailureReason;
+			}
+		}
+
+		for (FAssetRenameDataWithReferencers& RenameData : AssetsToRename)
+		{
+			if (!RenameData.bRenameFailed && !RenameData.bOnlyFixSoftReferences)
+			{
+				const FName PackageName = RenameData.OldObjectPath.GetLongPackageFName();
+				if (RenamedReferencingPackageNames.Contains(PackageName) || RenamedReferencedPackageNames.Contains(PackageName))
+				{
+					RenameData.bRenameFailed = true;
+					RenameData.FailureReason = FailureReason;
+				}
+			}
+		}
+	}
+	else
+	{
+		// Make assets public
+		for (UPackage* Package : PackagesToMakePublic)
+		{
+			Package->SetIsExternallyReferenceable(true);
+		}
+	}
+}
+
 void FAssetRenameManager::RenameReferencingSoftObjectPaths(const TArray<UPackage *> PackagesToCheck, const TMap<FSoftObjectPath, FSoftObjectPath>& AssetRedirectorMap) const
 {
 	// Add redirects as needed
@@ -1504,7 +1758,7 @@ void FAssetRenameManager::PerformAssetRename(TArray<FAssetRenameDataWithReferenc
 				/**
 				 * Do not save the package when the user simply changing a case and there is no referencer to it
 				 */
-				if (!(bIsCaseChangeOnly && RenameData.ReferencingPackageNames.IsEmpty()))
+				if (!(bIsCaseChangeOnly && RenameData.NotRenamedReferencingPackageNames.IsEmpty()))
 				{
 					PackagesToSave.AddUnique(Asset->GetOutermost());
 				}
@@ -1536,7 +1790,7 @@ void FAssetRenameManager::PerformAssetRename(TArray<FAssetRenameDataWithReferenc
 
 		if (!RenameData.bRenameFailed && !bIsCaseChangeOnly)
 		{
-			for (FName PackageName : RenameData.ReferencingPackageNames)
+			for (FName PackageName : RenameData.NotRenamedReferencingPackageNames)
 			{
 				UPackage* PackageToCheck = FindPackage(nullptr, *PackageName.ToString());
 				if (PackageToCheck)
diff --git a/Engine/Source/Developer/AssetTools/Private/AssetRenameManager.h b/Engine/Source/Developer/AssetTools/Private/AssetRenameManager.h
--- a/Engine/Source/Developer/AssetTools/Private/AssetRenameManager.h
+++ b/Engine/Source/Developer/AssetTools/Private/AssetRenameManager.h
@@ -107,6 +107,13 @@ private:
 	/** Finds any read only packages and removes them from the save list. Assets referenced by these packages will leave redirectors. */ 
 	void DetectReadOnlyPackages(TArray<FAssetRenameDataWithReferencers>& AssetsToRename, TArray<UPackage*>& InOutReferencingPackagesToSave) const;
 
+	/** 
+	 * Make public any asset that will be referenced from another plugin after the rename.
+	 * If the asset cannot be made public or if moving it requires a referenced asset that's not being modified to become public,
+	 * its rename will fail and other assets that have dependencies between them will also fail to be renamed.
+	 */
+	void SetupPublicAssets(TArray<FAssetRenameDataWithReferencers>& AssetsToRename) const;
+
 	/** Performs the asset rename after the user has selected to proceed */
 	void PerformAssetRename(TArray<FAssetRenameDataWithReferencers>& AssetsToRename) const;
 
diff --git a/Engine/Source/Developer/AssetTools/Private/AssetTools.cpp b/Engine/Source/Developer/AssetTools/Private/AssetTools.cpp
--- a/Engine/Source/Developer/AssetTools/Private/AssetTools.cpp
+++ b/Engine/Source/Developer/AssetTools/Private/AssetTools.cpp
@@ -448,6 +448,13 @@ namespace UE::AssetTools::Private
 				UE_LOG(LogAssetTools, Log, TEXT("Folder Write Permissions:\n%s"), *IAssetTools::Get().GetWritableFolderPermissionList()->ToString());
 			}));
 
+	static bool bEnablePublicAssetFeature = false;
+	static FAutoConsoleVariableRef CVarEnablePublicAssetFeature(
+		TEXT("AssetTools.EnablePublicAssetFeature"),
+		bEnablePublicAssetFeature,
+		TEXT("Enables the Experimental Public Asset Feature (False: disabled, True:enabled")
+	);
+
 	/** 
 	 * CVar to specify if we should use Header patching in advanced copy.
 	 * Default is false.
diff --git a/Engine/Source/Editor/ContentBrowser/Private/AssetContextMenu.cpp b/Engine/Source/Editor/ContentBrowser/Private/AssetContextMenu.cpp
--- a/Engine/Source/Editor/ContentBrowser/Private/AssetContextMenu.cpp
+++ b/Engine/Source/Editor/ContentBrowser/Private/AssetContextMenu.cpp
@@ -71,13 +71,6 @@
 class FMenuBuilder;
 class SWidget;
 
-
-namespace ContentBrowserConsoleVariables
-{
-	static bool ContentBrowser_EnablePublicAssetFeature = false;
-	static FAutoConsoleVariableRef ContentBrowser_EnablePublicAssetFeatureCVar(TEXT("ContentBrowser.PublicAsset.EnablePublicAssetFeature"),
-		ContentBrowser_EnablePublicAssetFeature, TEXT("Enables the Experimental Public Asset Feature (False: disabled, True:enabled"));
-}
 #define LOCTEXT_NAMESPACE "ContentBrowser"
 
 FAssetContextMenu::FAssetContextMenu(const TWeakPtr<SAssetView>& InAssetView)
@@ -434,7 +427,8 @@ void FAssetContextMenu::AddMenuOptions(UToolMenu* InMenu)
 	// Add quick access to view commands
 	AddExploreMenuOptions(InMenu);
 
-	if (ContentBrowserConsoleVariables::ContentBrowser_EnablePublicAssetFeature)
+	static const IConsoleVariable* EnablePublicAssetFeatureCVar = IConsoleManager::Get().FindConsoleVariable(TEXT("AssetTools.EnablePublicAssetFeature"));
+	if (EnablePublicAssetFeatureCVar && EnablePublicAssetFeatureCVar->GetBool())
 	{
 		AddPublicStateMenuOptions(InMenu);
 	}
diff --git a/Engine/Source/Editor/ContentBrowser/Private/AssetViewWidgets.cpp b/Engine/Source/Editor/ContentBrowser/Private/AssetViewWidgets.cpp
--- a/Engine/Source/Editor/ContentBrowser/Private/AssetViewWidgets.cpp
+++ b/Engine/Source/Editor/ContentBrowser/Private/AssetViewWidgets.cpp
@@ -889,8 +889,8 @@ TSharedRef<SWidget> SAssetViewItem::CreateToolTipWidget() const
 
 			TSharedRef<SVerticalBox> OverallTooltipVBox = SNew(SVerticalBox);
 
-			static const auto PublicAssetUIEnabledCVar = IConsoleManager::Get().FindConsoleVariable(TEXT("ContentBrowser.PublicAsset.EnablePublicAssetFeature"));
-			const bool bIsPublicAssetUIEnabled = PublicAssetUIEnabledCVar && PublicAssetUIEnabledCVar->GetBool();
+			static const IConsoleVariable* EnablePublicAssetFeatureCVar = IConsoleManager::Get().FindConsoleVariable(TEXT("AssetTools.EnablePublicAssetFeature"));
+			const bool bIsPublicAssetUIEnabled = EnablePublicAssetFeatureCVar && EnablePublicAssetFeatureCVar->GetBool();
 
 			// Top section (asset name, type, is checked out)
 			OverallTooltipVBox->AddSlot()
diff --git a/Engine/Source/Editor/ContentBrowser/Private/PathContextMenu.cpp b/Engine/Source/Editor/ContentBrowser/Private/PathContextMenu.cpp
--- a/Engine/Source/Editor/ContentBrowser/Private/PathContextMenu.cpp
+++ b/Engine/Source/Editor/ContentBrowser/Private/PathContextMenu.cpp
@@ -220,13 +220,8 @@ void FPathContextMenu::MakePathViewContextMenu(UToolMenu* Menu)
 				);
 			}
 
-			static const auto PublicAssetUIEnabledCVar = IConsoleManager::Get().FindConsoleVariable(TEXT("ContentBrowser.PublicAsset.EnablePublicAssetFeature"));
-			bool bIsPublicAssetUIEnabled = false;
-
-			if (PublicAssetUIEnabledCVar)
-			{
-				bIsPublicAssetUIEnabled = PublicAssetUIEnabledCVar->GetBool();
-			}
+			static const IConsoleVariable* EnablePublicAssetFeatureCVar = IConsoleManager::Get().FindConsoleVariable(TEXT("AssetTools.EnablePublicAssetFeature"));
+			const bool bIsPublicAssetUIEnabled = EnablePublicAssetFeatureCVar && EnablePublicAssetFeatureCVar->GetBool();
 
 			FStringView SelectedFolderPathView(SelectedFolderPath);
 			if (bIsPublicAssetUIEnabled && FContentBrowserSingleton::Get().IsFolderShowPrivateContentToggleable(SelectedFolderPathView))
diff --git a/Engine/Source/Editor/ContentBrowser/Private/SContentBrowser.cpp b/Engine/Source/Editor/ContentBrowser/Private/SContentBrowser.cpp
--- a/Engine/Source/Editor/ContentBrowser/Private/SContentBrowser.cpp
+++ b/Engine/Source/Editor/ContentBrowser/Private/SContentBrowser.cpp
@@ -3936,15 +3936,8 @@ void SContentBrowser::HandleItemDataUpdated(TArrayView<const FContentBrowserItem
 
 bool SContentBrowser::HandlePrivateContentFilter(const FContentBrowserItem& AssetItem)
 {
-	static const auto PublicAssetUIEnabledCVar = IConsoleManager::Get().FindConsoleVariable(TEXT("ContentBrowser.PublicAsset.EnablePublicAssetFeature"));
-	bool bIsPublicAssetUIEnabled = false;
-
-	if (PublicAssetUIEnabledCVar)
-	{
-		bIsPublicAssetUIEnabled = PublicAssetUIEnabledCVar->GetBool();
-	}
-
-	if (!bIsPublicAssetUIEnabled)
+	static const IConsoleVariable* EnablePublicAssetFeatureCVar = IConsoleManager::Get().FindConsoleVariable(TEXT("AssetTools.EnablePublicAssetFeature"));
+	if (!EnablePublicAssetFeatureCVar || !EnablePublicAssetFeatureCVar->GetBool())
 	{
 		return false;
 	}
