Using header patcher instead of root asset package patcher
Based on CL 30214572

Addeds Patching call to IAssetTools.
Consumes the call in a task in the DuplicatePluginCommandlet

diff --git a/Engine/Source/Developer/AssetTools/Private/AssetHeaderPatcher.cpp b/Engine/Source/Developer/AssetTools/Private/AssetHeaderPatcher.cpp
--- a/Engine/Source/Developer/AssetTools/Private/AssetHeaderPatcher.cpp
+++ b/Engine/Source/Developer/AssetTools/Private/AssetHeaderPatcher.cpp
@@ -253,10 +253,11 @@ class FAssetHeaderPatcherInner
 public:
 	using EResult = FAssetHeaderPatcher::EResult;
 
-	FAssetHeaderPatcherInner(const FString& InSrcAsset, const FString& InDstAsset, const TMap<FString, FString>& InSearchAndReplace, FArchive* InDstArchive = nullptr)
+	FAssetHeaderPatcherInner(const FString& InSrcAsset, const FString& InDstAsset, const TMap<FString, FString>& InSearchAndReplace, FArchive* InDstArchive = nullptr, bool bInBespokeSearchInUse = false)
 		: SrcAsset(InSrcAsset)
 		, DstAsset(InDstAsset)
 		, SearchAndReplace(InSearchAndReplace)
+		, bBespokeSearchInUse(bInBespokeSearchInUse)
 		, DstArchive(InDstArchive)
 	{
 	}
@@ -274,6 +275,7 @@ public:
 	const FString& SrcAsset;
 	const FString& DstAsset;
 	const TMap<FString, FString>& SearchAndReplace;
+	bool  bBespokeSearchInUse;
 	FArchive* DstArchive = nullptr;
 	TUniquePtr<FArchive> DstArchiveOwner;
 
@@ -325,9 +327,9 @@ public:
 	FAssetRegistryData AssetRegistryData;
 };
 
-FAssetHeaderPatcher::EResult FAssetHeaderPatcher::DoPatch(const FString& InSrcAsset, const FString& InDstAsset, const TMap<FString, FString>& InSearchAndReplace)
+FAssetHeaderPatcher::EResult FAssetHeaderPatcher::DoPatch(const FString& InSrcAsset, const FString& InDstAsset, const TMap<FString, FString>& InSearchAndReplace, bool bInStarRestrictionInUse)
 {
-	FAssetHeaderPatcherInner Inner(InSrcAsset, InDstAsset, InSearchAndReplace);
+	FAssetHeaderPatcherInner Inner(InSrcAsset, InDstAsset, InSearchAndReplace, nullptr, bInStarRestrictionInUse);
 
 	if (!FFileHelper::LoadFileToArray(Inner.SrcBuffer, *Inner.SrcAsset))
 	{
@@ -761,6 +763,16 @@ void FAssetHeaderPatcherInner::PatchHeader_PatchSections()
 					TagData.Value = PatchedAssetData;
 				}
 			}
+			else if (bBespokeSearchInUse &&
+					 FPathViews::GetBaseFilename(ObjData.ObjectData.ObjectPath) == TEXT("GameFeatureData") &&
+				     TagData.Key == TEXT("PrimaryAssetName"))
+			{
+				FString BespokeSearchValue = TEXT("<GameFeatureData.PrimaryAssetName>") + TagData.Value;
+				if (DoPatch(BespokeSearchValue)) 
+				{
+					TagData.Value = BespokeSearchValue;
+				}
+			}
 			else 
 			{
 				DoPatch(TagData.Value);
diff --git a/Engine/Source/Developer/AssetTools/Private/AssetHeaderPatcher.h b/Engine/Source/Developer/AssetTools/Private/AssetHeaderPatcher.h
--- a/Engine/Source/Developer/AssetTools/Private/AssetHeaderPatcher.h
+++ b/Engine/Source/Developer/AssetTools/Private/AssetHeaderPatcher.h
@@ -20,7 +20,7 @@ struct FAssetHeaderPatcher
 		ErrorEmptyRequireSection,
 	};
 
-	static EResult DoPatch(const FString& InSrcAsset, const FString& InDstAsset, const TMap<FString, FString>& InSearchAndReplace);
+	static EResult DoPatch(const FString& InSrcAsset, const FString& InDstAsse, const TMap<FString, FString>& InSearchAndReplace, bool bInStarRestrictionInUse = false);
 
 	static EResult Test_DoPatch(FArchive& InSrcReader, FArchive& InDstWriter, const TMap<FString, FString>& InSearchAndReplace);
 };
diff --git a/Engine/Source/Developer/AssetTools/Private/AssetTools.cpp b/Engine/Source/Developer/AssetTools/Private/AssetTools.cpp
--- a/Engine/Source/Developer/AssetTools/Private/AssetTools.cpp
+++ b/Engine/Source/Developer/AssetTools/Private/AssetTools.cpp
@@ -2334,7 +2334,7 @@ bool UAssetToolsImpl::AdvancedCopyPackages(
 				UE::Tasks::FTask PatcherTask = UE::Tasks::Launch(UE_SOURCE_LOCATION,
 					[&PatchAssetsCompletedCount, &PatchingPatterns, InSrcFilename = SrcFilename, InDestFilename = DestFilename, &ErroredFilesLock, &ErroredFiles] () 
 					{
-						FAssetHeaderPatcher::EResult Result = FAssetHeaderPatcher::DoPatch(InSrcFilename, InDestFilename, PatchingPatterns);
+						FAssetHeaderPatcher::EResult Result = FAssetHeaderPatcher::DoPatch(InSrcFilename, InDestFilename, PatchingPatterns, /* bBespokeSearchInUse */false);
 						if (Result != FAssetHeaderPatcher::EResult::Success) 
 						{
 							FScopeLock Lock(&ErroredFilesLock);
@@ -2692,6 +2692,54 @@ bool UAssetToolsImpl::AdvancedCopyPackages(const FAdvancedCopyParams& CopyParams
 	return bResult;
 }
 
+/** Copies a file, patching internal references without performing a de-serialization. This is a blocking operation. returns true on successful copy */
+bool UAssetToolsImpl::PatchCopyPackageFile(const FString& SrcFile, const FString& DstFile, const TMap<FString, FString>& SearchForAndReplace) const
+{
+	FAssetHeaderPatcher::EResult Result = FAssetHeaderPatcher::DoPatch(SrcFile, DstFile, SearchForAndReplace, /* bBespokeSearchInUse */ true);
+	return (Result == FAssetHeaderPatcher::EResult::Success);
+}
+
+TMap<FString, FString> UAssetToolsImpl::GetMappingsForRootPackageRename(
+	const FString& SrcRoot,
+	const FString& DstRoot,
+	const FString& SrcBaseDir,
+	const TArray<TPair<FString, FString>>& SourceAndDestFiles) const
+{
+	TMap<FString, FString> Result;
+	Result.Reserve(3 + SourceAndDestFiles.Num());
+
+	{	// Plugin name patterns
+		FString SrcPath = FPaths::Combine(TEXT("/"), SrcRoot, SrcRoot);
+		FString DstPath = FPaths::Combine(TEXT("/"), DstRoot, SrcRoot);
+
+		Result.Add(SrcPath + TEXT(".") + SrcRoot, DstPath + TEXT(".") + SrcRoot);	// /Src/Src.Src -> /Dst/Src.Src
+		Result.Add(MoveTemp(SrcPath), MoveTemp(DstPath));					// /Src/Src     -> /Dst/Src
+		Result.Add(TEXT("<GameFeatureData.PrimaryAssetName>") + SrcRoot, DstRoot);	// <GameFeatureData.PrimaryAssetName>Src -> Dst (for matching in specific modes)
+	}
+
+	const FString SourceContentPath = FPaths::Combine(SrcBaseDir, TEXT("Content"));
+
+	for (const TTuple<FString, FString>& SourceAndDest : SourceAndDestFiles)
+	{
+		const FString& SrcFileName = SourceAndDest.Key;
+
+		if (FPaths::IsUnderDirectory(SrcFileName, SourceContentPath))
+		{
+			if (FStringView RelativePkgPath; FPathViews::TryMakeChildPathRelativeTo(SrcFileName, SourceContentPath, RelativePkgPath))
+			{
+				RelativePkgPath = FPathViews::GetBaseFilenameWithPath(RelativePkgPath); // chop the extension
+				if (RelativePkgPath.Len() > 0 && !RelativePkgPath.EndsWith(TEXT("/")))
+				{
+					Result.Add(FPaths::Combine(TEXT("/"), SrcRoot, RelativePkgPath),
+						       FPaths::Combine(TEXT("/"), DstRoot, RelativePkgPath));
+				}
+			}
+		}
+	}
+
+	return Result;
+}
+
 bool UAssetToolsImpl::IsDiscoveringAssetsInProgress() const
 {
 	FAssetRegistryModule& AssetRegistryModule = FModuleManager::LoadModuleChecked<FAssetRegistryModule>("AssetRegistry");
diff --git a/Engine/Source/Developer/AssetTools/Private/AssetTools.h b/Engine/Source/Developer/AssetTools/Private/AssetTools.h
--- a/Engine/Source/Developer/AssetTools/Private/AssetTools.h
+++ b/Engine/Source/Developer/AssetTools/Private/AssetTools.h
@@ -123,6 +123,10 @@ public:
 	virtual void ExpandDirectories(const TArray<FString>& Files, const FString& DestinationPath, TArray<TPair<FString, FString>>& FilesAndDestinations) const override;
 	virtual bool AdvancedCopyPackages(const FAdvancedCopyParams& CopyParams, const TArray<TMap<FString, FString>>& PackagesAndDestinations) const override;
 	virtual bool AdvancedCopyPackages(const TMap<FString, FString>& SourceAndDestPackages, const bool bForceAutosave, const bool bCopyOverAllDestinationOverlaps, FDuplicatedObjects* OutDuplicatedObjects, EMessageSeverity::Type NotificationSeverityFilter) const override;
+	virtual bool PatchCopyPackageFile(const FString& SrcFile, const FString& DstFile, const TMap<FString, FString>& SearchForAndReplace) const;
+	virtual TMap<FString, FString> GetMappingsForRootPackageRename(const FString& SrcRoot, const FString& DstRoot, const FString& SrcBaseDir, const TArray<TPair<FString, FString>>& SourceAndDestFiles) const;
+
+
 	virtual void GenerateAdvancedCopyDestinations(FAdvancedCopyParams& InParams, const TArray<FName>& InPackageNamesToCopy, const class UAdvancedCopyCustomization* CopyCustomization, TMap<FString, FString>& OutPackagesAndDestinations) const override;
 	virtual bool FlattenAdvancedCopyDestinations(const TArray<TMap<FString, FString>>& PackagesAndDestinations, TMap<FString, FString>& FlattenedPackagesAndDestinations) const override;
 	virtual bool ValidateFlattenedAdvancedCopyDestinations(const TMap<FString, FString>& FlattenedPackagesAndDestinations) const override;
diff --git a/Engine/Source/Developer/AssetTools/Public/IAssetTools.h b/Engine/Source/Developer/AssetTools/Public/IAssetTools.h
--- a/Engine/Source/Developer/AssetTools/Public/IAssetTools.h
+++ b/Engine/Source/Developer/AssetTools/Public/IAssetTools.h
@@ -551,6 +551,12 @@ public:
 	/** Copies files after the flattened map of sources and destinations was confirmed */
 	virtual bool AdvancedCopyPackages(const TMap<FString, FString>& SourceAndDestPackages, const bool bForceAutosave = false, const bool bCopyOverAllDestinationOverlaps = true, FDuplicatedObjects* OutDuplicatedObjects = nullptr, EMessageSeverity::Type NotificationSeverityFilter = EMessageSeverity::Info) const = 0;
 
+	/** Copies a file, patching internal references without performing a de-serialization. This is a blocking operation. returns true on successful copy */
+	virtual bool PatchCopyPackageFile(const FString& SrcFile, const FString& DstFile, const TMap<FString, FString>& SearchForAndReplace) const = 0;
+
+	/** Generates the SearchAndReplace map for a PatchCopyPackageFile if all you are doing is changing the root and not the name. */
+	virtual TMap<FString, FString> GetMappingsForRootPackageRename(const FString& SrcRoot, const FString& DstRoot, const FString& SrcBaseDir, const TArray<TPair<FString, FString>>& SourceAndDestFiles) const = 0;
+
 	/* Given a set of packages to copy, generate the map of those packages to destination filenames */
 	virtual void GenerateAdvancedCopyDestinations(FAdvancedCopyParams& InParams, const TArray<FName>& InPackageNamesToCopy, const UAdvancedCopyCustomization* CopyCustomization, TMap<FString, FString>& OutPackagesAndDestinations) const = 0;
 
