Added scope timer to indicate use of header patching

diff --git a/Engine/Source/Developer/AssetTools/Private/AssetTools.cpp b/Engine/Source/Developer/AssetTools/Private/AssetTools.cpp
--- a/Engine/Source/Developer/AssetTools/Private/AssetTools.cpp
+++ b/Engine/Source/Developer/AssetTools/Private/AssetTools.cpp
@@ -2296,6 +2296,8 @@ bool UAssetToolsImpl::AdvancedCopyPackages(
 
 		if (UE::AssetTools::Private::bEnableHeaderPatching)
 		{
+			TRACE_CPUPROFILER_EVENT_SCOPE(AdvancedCopyPackages.HeaderPatching);
+
 			std::atomic<int32> PatchAssetsCompletedCount = 0;
 			UE::Tasks::FTaskEvent PatchAssetsCompletionTask{ UE_SOURCE_LOCATION };
 
