Fix for Headerpatching and Default__$name_C objects

diff --git a/Engine/Source/Developer/AssetTools/Private/AssetTools.cpp b/Engine/Source/Developer/AssetTools/Private/AssetTools.cpp
--- a/Engine/Source/Developer/AssetTools/Private/AssetTools.cpp
+++ b/Engine/Source/Developer/AssetTools/Private/AssetTools.cpp
@@ -2255,6 +2255,7 @@ TMap<FString, FString> GenerateAdditionalAssetMappings(const TMap<FString, FStri
 		{
 			Result.Add({ FString(SrcPackageName), FString(DstPackageName) });
 			Result.Add({ FString(SrcPackageName) + TEXT("_C"), FString(DstPackageName) + TEXT("_C") }); // catch compiled blueprint names
+			Result.Add({ DEFAULT_OBJECT_PREFIX + FString(SrcPackageName) + TEXT("_C"), DEFAULT_OBJECT_PREFIX + FString(DstPackageName) + TEXT("_C") }); // BPGC default object
 		}
 	}
 
