Added warning about packages that cannot be found when wanting to be copied.

diff --git a/Engine/Source/Developer/AssetTools/Private/AssetTools.cpp b/Engine/Source/Developer/AssetTools/Private/AssetTools.cpp
--- a/Engine/Source/Developer/AssetTools/Private/AssetTools.cpp
+++ b/Engine/Source/Developer/AssetTools/Private/AssetTools.cpp
@@ -2319,6 +2319,10 @@ bool UAssetToolsImpl::AdvancedCopyPackages(
 				{
 					FString DestFilename = FPackageName::LongPackageNameToFilename(DestPackage, FString(FPathViews::GetExtension(SrcFilename, true)));
 					ToCopyAndPatchFiles.Add({ MoveTemp(SrcFilename), MoveTemp(DestFilename) });
+				} 
+				else
+				{
+					UE_LOG(LogAssetTools, Warning, TEXT("{%s} package does not exist, and will not be copied."), *PackageName);
 				}
 			}
 
