AssetViewUtils::MoveAssets - Fix issue whe external packages weren't moved
* Previous fix was skipped for non local assets that needed a redirector

diff --git a/Engine/Source/Developer/AssetTools/Private/AssetRenameManager.cpp b/Engine/Source/Developer/AssetTools/Private/AssetRenameManager.cpp
--- a/Engine/Source/Developer/AssetTools/Private/AssetRenameManager.cpp
+++ b/Engine/Source/Developer/AssetTools/Private/AssetRenameManager.cpp
@@ -872,9 +872,19 @@ void FAssetRenameManager::LoadReferencingPackages(TArray<FAssetRenameDataWithRef
 		
 		FAssetRenameDataWithReferencers& RenameData = AssetsToRename[AssetIdx];
 
+		TSet<FName> ReferencingExternalPackageNames;
+
 		UObject* Asset = RenameData.Asset.Get();
 		if (Asset)
 		{
+			// External packages must always be resaved
+			for (UPackage* ExternalPackage : Asset->GetPackage()->GetExternalPackages())
+			{
+				FName ExternalPackageName = ExternalPackage->GetFName();
+				ReferencingExternalPackageNames.Add(ExternalPackageName);
+				OutReferencingPackagesToSave.Add(ExternalPackage);
+			}
+
 			// Make sure this asset is local. Only local assets should be renamed without a redirector
 			if (bCheckStatus)
 			{
@@ -921,17 +931,6 @@ void FAssetRenameManager::LoadReferencingPackages(TArray<FAssetRenameDataWithRef
 		TArray<UPackage*> PackagesToSaveForThisAsset;
 		bool bAllPackagesLoadedForThisAsset = true;
 
-		TSet<FName> ReferencingExternalPackageNames;
-		if (Asset)
-		{
-			for (UPackage* ExternalPackage : Asset->GetPackage()->GetExternalPackages())
-			{
-				FName ExternalPackageName = ExternalPackage->GetFName();
-				ReferencingExternalPackageNames.Add(ExternalPackageName);
-				OutReferencingPackagesToSave.Add(ExternalPackage);
-			}
-		}
-
 		for (auto It = RenameData.NotRenamedReferencingPackageNames.CreateIterator(); It; ++It)
 		{
 			FName PackageName = *It;
