Silencing warning about not copying Verse packages in advanced copy with header patching.

diff --git a/Engine/Source/Developer/AssetTools/Private/AssetTools.cpp b/Engine/Source/Developer/AssetTools/Private/AssetTools.cpp
--- a/Engine/Source/Developer/AssetTools/Private/AssetTools.cpp
+++ b/Engine/Source/Developer/AssetTools/Private/AssetTools.cpp
@@ -2315,6 +2315,14 @@ bool UAssetToolsImpl::AdvancedCopyPackages(
 				const FString& DestPackage = Package.Value;
 				FString SrcFilename;
 
+				if (FPackageName::IsVersePackage(PackageName))
+				{
+					// Verse packages are not header patchable.
+					// They are also not Packages as far as DoesPackageExist tells me.
+					// But they are real files that in template copying have already been done, so we dont want a warning message.
+					continue;
+				}
+
 				if (FPackageName::DoesPackageExist(PackageName, &SrcFilename))
 				{
 					FString DestFilename = FPackageName::LongPackageNameToFilename(DestPackage, FString(FPathViews::GetExtension(SrcFilename, true)));
