Fix up redirectors flow update
- Rename menu entries to Update (Redirector) Referencers
- No separate options for fixup vs fixup and keep
- After fixup, prompt user to delete redirectors if project settings allow
- After fixup, show users what files could not be checked out/saved and what redirectors remain

diff --git a/Engine/Plugins/Editor/ContentBrowser/ContentBrowserAssetDataSource/Source/ContentBrowserAssetDataSource/Private/AssetFolderContextMenu.cpp b/Engine/Plugins/Editor/ContentBrowser/ContentBrowserAssetDataSource/Source/ContentBrowserAssetDataSource/Private/AssetFolderContextMenu.cpp
--- a/Engine/Plugins/Editor/ContentBrowser/ContentBrowserAssetDataSource/Source/ContentBrowserAssetDataSource/Private/AssetFolderContextMenu.cpp
+++ b/Engine/Plugins/Editor/ContentBrowser/ContentBrowserAssetDataSource/Source/ContentBrowserAssetDataSource/Private/AssetFolderContextMenu.cpp
@@ -47,11 +47,10 @@ void FAssetFolderContextMenu::AddMenuOptions(UToolMenu* Menu)
 			{
 				FToolMenuEntry& Entry = Section.AddMenuEntry(
 					"FixUpRedirectorsInFolder",
-					LOCTEXT("FixUpRedirectorsInFolder", "Fix Up Redirectors"),
-					LOCTEXT("FixUpRedirectorsInFolderTooltip", "Finds referencers to all redirectors in the selected items and resaves them if possible, then deletes any redirectors that had all their referencers fixed."),
+					LOCTEXT("FixUpRedirectorsInFolder", "Update Redirector References"),
+					LOCTEXT("FixUpRedirectorsInFolderTooltip", "Finds references to all redirectors in the selected items and resaves the referencing assets if possible, so that they reference the target of the redirector directly instead."),
 					FSlateIcon(FAppStyle::GetAppStyleSetName(), "Icons.Adjust"),
-					FUIAction(FExecuteAction::CreateSP(this, &FAssetFolderContextMenu::ExecuteFixUpRedirectorsInFolder))
-					);
+					FUIAction(FExecuteAction::CreateSP(this, &FAssetFolderContextMenu::ExecuteFixUpRedirectorsInFolder)));
 				Entry.InsertPosition = FToolMenuInsert("Delete", EToolMenuInsertType::After);
 			}
 
@@ -173,7 +172,7 @@ void FAssetFolderContextMenu::ExecuteFixUpRedirectorsInFolder()
 		SlowTask.EnterProgressFrame(1, LOCTEXT("FixupRedirectors_FixupReferencers", "Fixing up referencers..."));
 		// Load the asset tools module
 		FAssetToolsModule& AssetToolsModule = FModuleManager::LoadModuleChecked<FAssetToolsModule>(TEXT("AssetTools"));
-		AssetToolsModule.Get().FixupReferencers(Redirectors);
+		AssetToolsModule.Get().FixupReferencers(Redirectors, true, ERedirectFixupMode::PromptForDeletingRedirectors);
 	}
 }
 
diff --git a/Engine/Plugins/Editor/EngineAssetDefinitions/Source/Private/AssetDefinition_Redirector.cpp b/Engine/Plugins/Editor/EngineAssetDefinitions/Source/Private/AssetDefinition_Redirector.cpp
--- a/Engine/Plugins/Editor/EngineAssetDefinitions/Source/Private/AssetDefinition_Redirector.cpp
+++ b/Engine/Plugins/Editor/EngineAssetDefinitions/Source/Private/AssetDefinition_Redirector.cpp
@@ -44,7 +44,7 @@ namespace MenuExtension_Redirector
 		}
 	}
 
-	static void ExecuteFixUp(const FToolMenuContext& MenuContext, bool bDeleteAssets)
+	static void ExecuteFixUp(const FToolMenuContext& MenuContext)
 	{
 		// This will fix references to selected redirectors, except in the following cases:
 		// Redirectors referenced by unloaded maps will not be fixed up, but any references to it that can be fixed up will
@@ -59,7 +59,7 @@ namespace MenuExtension_Redirector
 			if (Redirectors.Num() > 0)
 			{
 				IAssetTools& AssetTools = FModuleManager::LoadModuleChecked<FAssetToolsModule>("AssetTools").Get();
-				AssetTools.FixupReferencers(Redirectors, /*bCheckoutDialogPrompt=*/ true, bDeleteAssets ? ERedirectFixupMode::DeleteFixedUpRedirectors : ERedirectFixupMode::LeaveFixedUpRedirectors);
+				AssetTools.FixupReferencers(Redirectors, /*bCheckoutDialogPrompt=*/true, ERedirectFixupMode::PromptForDeletingRedirectors);
 			}
 		}
 	}
@@ -82,20 +82,12 @@ namespace MenuExtension_Redirector
 					InSection.AddMenuEntry("Redirector_FindTarget", Label, ToolTip, Icon, UIAction);
 				}
 				{
-					const TAttribute<FText> Label = LOCTEXT("Redirector_FixUp", "Fix Up");
-					const TAttribute<FText> ToolTip = LOCTEXT("Redirector_FixUpTooltip", "Finds referencers to selected redirectors and resaves them if possible, then deletes any redirectors that had all their referencers fixed.");
+					const TAttribute<FText> Label = LOCTEXT("Redirector_UpdateReferencers", "Update References");
+					const TAttribute<FText> ToolTip = LOCTEXT("Redirector_FixUpTooltip", "Finds references to selected redirectors and resaves the referencing assets if possible, so that they reference the target of the redirector directly instead.");
 					const FSlateIcon Icon = FSlateIcon(FAppStyle::GetAppStyleSetName(), "ClassIcon.ObjectRedirector");
-					const FToolMenuExecuteAction UIAction = FToolMenuExecuteAction::CreateStatic(&ExecuteFixUp, true);
+					const FToolMenuExecuteAction UIAction = FToolMenuExecuteAction::CreateStatic(&ExecuteFixUp);
 
-					InSection.AddMenuEntry("Redirector_FixUp", Label, ToolTip, Icon, UIAction);
-				}
-				{
-					const TAttribute<FText> Label = LOCTEXT("Redirector_FixUp_KeepingRedirector", "Fix Up (Keep Redirector)");
-					const TAttribute<FText> ToolTip = LOCTEXT("Redirector_FixUp_KeepingRedirectorTooltip", "Finds referencers to selected redirectors and resaves them if possible.");
-					const FSlateIcon Icon = FSlateIcon(FAppStyle::GetAppStyleSetName(), "ClassIcon.MyAsset");
-					const FToolMenuExecuteAction UIAction = FToolMenuExecuteAction::CreateStatic(&ExecuteFixUp, false);
-
-					InSection.AddMenuEntry("Redirector_FixUp_KeepingRedirector", Label, ToolTip, Icon, UIAction);
+					InSection.AddMenuEntry("Redirector_UpdateReferencers", Label, ToolTip, Icon, UIAction);
 				}
 			}));
 		}));
diff --git a/Engine/Source/Developer/AssetTools/Private/AssetFixUpRedirectors.cpp b/Engine/Source/Developer/AssetTools/Private/AssetFixUpRedirectors.cpp
--- a/Engine/Source/Developer/AssetTools/Private/AssetFixUpRedirectors.cpp
+++ b/Engine/Source/Developer/AssetTools/Private/AssetFixUpRedirectors.cpp
@@ -1,48 +1,519 @@
 // Copyright Epic Games, Inc. All Rights Reserved.
 
-
 #include "AssetFixUpRedirectors.h"
-#include "UObject/ObjectRedirector.h"
-#include "Misc/MessageDialog.h"
+#include "Algo/Sort.h"
+#include "AssetRegistry/AssetRegistryModule.h"
+#include "AssetRenameManager.h"
+#include "AssetTools.h"
+#include "CollectionManagerModule.h"
+#include "Dialogs/Dialogs.h"
+#include "Engine/Blueprint.h"
+#include "FileHelpers.h"
 #include "HAL/FileManager.h"
+#include "ICollectionManager.h"
+#include "ISourceControlModule.h"
+#include "ISourceControlOperation.h"
+#include "Logging/MessageLog.h"
+#include "Misc/MessageDialog.h"
+#include "Misc/PackageName.h"
+#include "Misc/ScopeExit.h"
 #include "Misc/ScopedSlowTask.h"
 #include "Modules/ModuleManager.h"
-#include "UObject/UObjectHash.h"
-#include "UObject/MetaData.h"
-#include "Misc/PackageName.h"
-#include "ISourceControlOperation.h"
-#include "SourceControlOperations.h"
-#include "ISourceControlModule.h"
-#include "FileHelpers.h"
-#include "SDiscoveringAssetsDialog.h"
-#include "AssetRenameManager.h"
-#include "AssetRegistry/AssetRegistryModule.h"
-#include "ICollectionManager.h"
-#include "CollectionManagerModule.h"
 #include "ObjectTools.h"
-#include "Logging/MessageLog.h"
-#include "AssetTools.h"
-#include "Engine/Blueprint.h"
+#include "PackageTools.h"
+#include "RevisionControlStyle/RevisionControlStyle.h"
+#include "SDiscoveringAssetsDialog.h"
+#include "Settings/EditorProjectSettings.h"
 #include "SourceControlHelpers.h"
+#include "SourceControlOperations.h"
+#include "UObject/MetaData.h"
+#include "UObject/ObjectRedirector.h"
+#include "UObject/StrongObjectPtr.h"
+#include "UObject/UObjectHash.h"
+#include "Widgets/Input/SButton.h"
+#include "Widgets/Views/STreeView.h"
 
 #define LOCTEXT_NAMESPACE "AssetFixUpRedirectors"
 
-
 struct FRedirectorRefs
 {
-	UObjectRedirector* Redirector;
+	TStrongObjectPtr<const UObjectRedirector> Redirector;
 	FName RedirectorPackageName;
 	TArray<FName> ReferencingPackageNames;
-	FText FailureReason;
-	bool bRedirectorValidForFixup;
 
-	FRedirectorRefs(UObjectRedirector* InRedirector)
+	// Referencing packages which could not be updated to remove references to this redirector
+	TArray<FName> LockedReferencerPackageNames;
+	TArray<FName> FailedReferencerPackageNames;
+
+	// Possible failure reason unrelated to above lists of package names.
+	// If this is not empty we may not delete the redirector - we may still be able to remove some references to it.
+	TArray<FText> OtherFailures;
+
+	bool bSCCError = false; // Failure to check the redirector itself out of source control etc
+
+	explicit FRedirectorRefs(FName PackageName)
+		: Redirector(nullptr)
+		, RedirectorPackageName(PackageName)
+	{
+	}
+
+	explicit FRedirectorRefs(const UObjectRedirector* InRedirector)
 		: Redirector(InRedirector)
 		, RedirectorPackageName(InRedirector->GetOutermost()->GetFName())
-		, bRedirectorValidForFixup(true)
-	{}
+	{
+	}
+};
+
+// What the user chooses to do from the report dialog
+enum class EFixupReportDecision
+{
+	None,
+	DeleteUnreferencedRedirectors
+};
+
+enum class EFixupReportItemType
+{
+	Redirector,             // A redirector which has remaining referencers
+	UnreferencedRedirector, // A redirector which has no remaining referencers
+	LockedRedirector,       // A redirector that could not be checked out of source control
+	LockedReferencer,       // A referencer which could not be checked out of source control
+	FailedReferencer,       // A referencer which could not be saved for some reason
+	Message,                // Another reason this redirector could not be deleted
+};
+
+// An item presented in the report dialog
+struct FFixupReportItem
+{
+	FText Text;
+	EFixupReportItemType Type;
+	TArray<TSharedPtr<FFixupReportItem>> Children;
+
+	explicit FFixupReportItem(EFixupReportItemType InType, FText InText)
+		: Text(MoveTemp(InText))
+		, Type(InType)
+	{
+	}
+
+	const FSlateBrush* GetIcon() const
+	{
+		switch (Type)
+		{
+			case EFixupReportItemType::Redirector:
+				return FAppStyle::Get().GetBrush("Icons.WarningWithColor");
+			case EFixupReportItemType::UnreferencedRedirector:
+				return FAppStyle::Get().GetBrush("Icons.SuccessWithColor");
+			case EFixupReportItemType::LockedRedirector:
+				// Share icon type for both kinds of locked items, through
+			case EFixupReportItemType::LockedReferencer:
+				return FRevisionControlStyleManager::Get().GetBrush("RevisionControl.Locked");
+			case EFixupReportItemType::FailedReferencer:
+				return FAppStyle::Get().GetBrush("Icons.ErrorWithColor");
+			case EFixupReportItemType::Message:
+				return FAppStyle::Get().GetBrush("Icons.ErrorWithColor");
+			default:
+				checkf(false, TEXT("Unhandled EFixupReportItemType"));
+				return nullptr;
+		}
+	}
+
+	FText GetName() const
+	{
+		switch (Type)
+		{
+			default:
+				// For redirectors and referencers, the Text is the asset name/path
+				return Text;
+			case EFixupReportItemType::Message:
+				// For messages the Text is the error message so show 'Other'
+				return LOCTEXT("Other", "Other");
+		}
+	}
+
+	FText GetDescription() const
+	{
+		switch (Type)
+		{
+			case EFixupReportItemType::Redirector:
+				return FText::Format(LOCTEXT("RedirectorIssueCount", "{0} issues"), Children.Num());
+			case EFixupReportItemType::UnreferencedRedirector:
+				return LOCTEXT("UnreferencedRedirector", "All references to this redirector have been removed.");
+			case EFixupReportItemType::LockedRedirector:
+				return LOCTEXT("LockedRedirector", "Redirector could not be checked out or marked for delete.");
+			case EFixupReportItemType::LockedReferencer:
+				return LOCTEXT("LockedReferencer", "Referencer could not be checked out of source control.");
+			case EFixupReportItemType::FailedReferencer:
+				return LOCTEXT("FailedReferencer", "Referencer could not be saved successfully.");
+			case EFixupReportItemType::Message:
+				return Text;
+			default:
+				checkf(false, TEXT("Unhandled EFixupReportItemType"));
+				return FText::GetEmpty();
+		}
+	}
+};
+
+namespace FixupRedirectorsReport
+{
+const FName ColumnID_Expander("Expander");
+const FName ColumnID_Icon("Icon");
+const FName ColumnID_Name("Name");
+const FName ColumnID_Description("Description");
+} // namespace FixupRedirectorsReport
+
+class SFixupRedirectorsReportTableRow : public SMultiColumnTableRow<TSharedPtr<FFixupReportItem>>
+{
+public:
+	SLATE_BEGIN_ARGS(SFixupRedirectorsReportTableRow) { }
+	SLATE_END_ARGS()
+
+	void Construct(const FArguments& InArgs, const TSharedRef<STableViewBase>& InOwnerTableView, TSharedPtr<FFixupReportItem> InItem)
+	{
+		Item = InItem;
+		SMultiColumnTableRow<TSharedPtr<FFixupReportItem>>::Construct(FSuperRowType::FArguments().Padding(0.0f), InOwnerTableView);
+	}
+
+	virtual TSharedRef<SWidget> GenerateWidgetForColumn(const FName& ColumnName)
+	{
+		if (ColumnName == FixupRedirectorsReport::ColumnID_Expander)
+		{
+			return SAssignNew(ExpanderArrowWidget, SExpanderArrow, SharedThis(this))
+				.StyleSet(ExpanderStyleSet)
+				.ShouldDrawWires(true);
+		}
+		else if (ColumnName == FixupRedirectorsReport::ColumnID_Icon)
+		{
+			return SNew(SImage)
+				.Image(Item->GetIcon());
+		}
+		else if (ColumnName == FixupRedirectorsReport::ColumnID_Name)
+		{
+			// clang-format off
+			return SNew(SHorizontalBox)
+				+SHorizontalBox::Slot()
+				.AutoWidth()
+				.Padding(3, 0, 0, 0)
+				[
+					SNew(STextBlock)
+					.Text(Item->GetName())
+				];
+			// clang-format on
+		}
+		else if (ColumnName == FixupRedirectorsReport::ColumnID_Description)
+		{
+			return SNew(STextBlock)
+				.Text(Item->GetDescription());
+		}
+		return SNullWidget::NullWidget;
+	}
+
+private:
+	TSharedPtr<FFixupReportItem> Item;
 };
 
+class SFixupRedirectorsReport : public SModalEditorDialog<bool>
+{
+public:
+	SLATE_BEGIN_ARGS(SFixupRedirectorsReport)
+	{
+	}
+	SLATE_END_ARGS()
+
+	void Construct(const FArguments& InArgs, const TArray<FRedirectorRefs>& RedirectorRefsList, bool InCanDelete)
+	{
+		bCanDelete = InCanDelete;
+		Sizing = ESizingRule::UserSized;
+		MinDimensions = FVector2D(1000.0f, 600.f);
+
+		bool bAnyRedirectorsToDelete = false;
+		for (const FRedirectorRefs& Redirector : RedirectorRefsList)
+		{
+			TSharedPtr<FFixupReportItem>& NewItem = TreeItems.Emplace_GetRef(MakeShared<FFixupReportItem>(
+				EFixupReportItemType::Redirector,
+				FText::FromName(Redirector.RedirectorPackageName)));
+			for (FName LockedPackageName : Redirector.LockedReferencerPackageNames)
+			{
+				NewItem->Children.Emplace(MakeShared<FFixupReportItem>(
+					EFixupReportItemType::LockedReferencer,
+					FText::FromName(LockedPackageName)));
+			}
+			for (FName FailedPackageName : Redirector.FailedReferencerPackageNames)
+			{
+				NewItem->Children.Emplace(MakeShared<FFixupReportItem>(
+					EFixupReportItemType::FailedReferencer,
+					FText::FromName(FailedPackageName)));
+			}
+			for (FText Issue : Redirector.OtherFailures)
+			{
+				NewItem->Children.Emplace(MakeShared<FFixupReportItem>(
+					EFixupReportItemType::Message,
+					Issue));
+			}
+
+			if (Redirector.bSCCError)
+			{
+				NewItem->Type = EFixupReportItemType::LockedRedirector;
+				bAnyProblems = true;
+			}
+			else if (NewItem->Children.IsEmpty())
+			{
+				NewItem->Type = EFixupReportItemType::UnreferencedRedirector;
+				bAnyRedirectorsToDelete = true;
+			}
+			else
+			{
+				bAnyProblems = true;
+			}
+		}
+
+		if (bCanDelete && bAnyRedirectorsToDelete)
+		{
+			MinDimensions.Y = 800.0f; // Slightly bigger default size to account for help text
+		}
+
+		Algo::SortBy(
+			TreeItems,
+			[](const TSharedPtr<FFixupReportItem>& Item) -> int32 { return !!Item->Children.Num(); },
+			TGreater<int32>{});
+
+		TSharedRef<SHeaderRow> HeaderRowWidget =
+			SNew(SHeaderRow)
+			+ SHeaderRow::Column(FixupRedirectorsReport::ColumnID_Icon)
+				  .DefaultLabel(FText::GetEmpty())
+				  .HAlignHeader(EHorizontalAlignment::HAlign_Left)
+				  .FixedWidth(20.0f)
+				  .HAlignCell(HAlign_Center)
+			+ SHeaderRow::Column(FixupRedirectorsReport::ColumnID_Expander)
+				  .DefaultLabel(FText::GetEmpty())
+				  .HAlignHeader(EHorizontalAlignment::HAlign_Left)
+				  .FixedWidth(16.0f)
+			+ SHeaderRow::Column(FixupRedirectorsReport::ColumnID_Name)
+				  .DefaultLabel(LOCTEXT("Column_Name", "Name"))
+				  .HAlignHeader(EHorizontalAlignment::HAlign_Left)
+				  .FillWidth(1.0f)
+			+ SHeaderRow::Column(FixupRedirectorsReport::ColumnID_Description)
+				  .DefaultLabel(FText::GetEmpty())
+				  .HAlignHeader(EHorizontalAlignment::HAlign_Left)
+				  .FillWidth(1.0f);
+
+		// clang-format off
+		TSharedPtr<SVerticalBox> VerticalBox = SNew(SVerticalBox)
+			+ SVerticalBox::Slot()
+			.FillHeight(1.0f)
+			[
+				SNew(SBorder)
+				.BorderImage(FAppStyle::GetBrush("ToolPanel.GroupBorder"))
+				[
+					SNew(SVerticalBox)
+					+ SVerticalBox::Slot()
+					.AutoHeight()
+					.Padding(0,0,0,10)
+					[
+						SNew(SBorder)
+						.BorderImage(FAppStyle::GetBrush("DetailsView.CategoryTop"))
+						.BorderBackgroundColor(FLinearColor(.6, .6, 0.6, 1.0f))
+						[
+							SNew(STextBlock)
+							.Text(this, &SFixupRedirectorsReport::GetIntroText)
+							.Font(FAppStyle::GetFontStyle("BoldFont"))
+							.ShadowOffset(FVector2D(1.0f, 1.0f))
+						]
+					]
+					+ SVerticalBox::Slot()
+					.FillHeight(1.0f)
+					[
+						SAssignNew(TreeView, STreeView<TSharedPtr<FFixupReportItem>>)
+						.HeaderRow(HeaderRowWidget)
+						.TreeItemsSource(&TreeItems)
+						.SelectionMode(ESelectionMode::None)
+						.OnGenerateRow(this, &SFixupRedirectorsReport::HandleGenerateTreeViewRow)
+						.OnGetChildren(this, &SFixupRedirectorsReport::HandleGetChildren)
+					]	
+				]
+			];
+		if (bCanDelete && bAnyRedirectorsToDelete)
+		{
+			// Large bottom row with scary delete buttons and long explanations
+			VerticalBox->AddSlot()
+				.AutoHeight()
+				[
+					SNew(SHorizontalBox)
+					+SHorizontalBox::Slot()
+					.FillWidth(1.0f)
+					.VAlign(VAlign_Bottom)
+					[
+						SNew(SVerticalBox)
+						+ SVerticalBox::Slot()
+						.FillHeight(1.0f)
+						[
+							SNew(SBorder)
+							.BorderImage(FAppStyle::GetBrush("NoBorder"))
+							[
+								SNew(STextBlock)
+								.Text(this, &SFixupRedirectorsReport::GetDeleteExplanationText)
+								.AutoWrapText(true)
+							]
+						]
+						+ SVerticalBox::Slot()
+						.AutoHeight()
+						[
+							SNew(SBorder)
+							.BorderImage(FAppStyle::GetBrush("NoBorder"))
+							.VAlign(VAlign_Bottom)
+							[
+								SNew(SButton)
+								.HAlign(HAlign_Center)
+								.Text(LOCTEXT("DeleteRedirectorsButtonText", "Delete Unreferenced Redirectors"))
+								.ToolTipText(LOCTEXT("DeleteTooltipText", "Delete redirectors which are no longer referenced"))
+								.ButtonStyle(FAppStyle::Get(), "FlatButton.Danger")
+								.TextStyle(FAppStyle::Get(), "FlatButton.DefaultTextStyle")
+								.OnClicked(this, &SFixupRedirectorsReport::CloseDeleteRedirectors)
+							]
+						]
+					]
+
+					+SHorizontalBox::Slot()
+					.FillWidth(1.0f)
+					.Padding(0, 10, 0, 0)
+					[
+						SNew(SVerticalBox)
+						+ SVerticalBox::Slot()
+						.FillHeight(1.0f)
+						[
+							SNew(SBorder)
+							.BorderImage(FAppStyle::GetBrush("NoBorder"))
+							// Only show the help text for keeping if deletion is also an option
+							[
+								SNew(STextBlock)
+								.Text(this, &SFixupRedirectorsReport::GetKeepExplanationText)
+								.AutoWrapText(true)
+							]
+						]
+						+ SVerticalBox::Slot()
+						.AutoHeight()
+						[
+							SNew(SBorder)
+							.BorderImage(FAppStyle::GetBrush("NoBorder"))
+							.VAlign(VAlign_Bottom)
+							[
+								SAssignNew(KeepButton, SButton)
+								.HAlign(HAlign_Center)
+								.Text(this, &SFixupRedirectorsReport::GetKeepRedirectorsButtonText)
+								.ToolTipText(this, &SFixupRedirectorsReport::GetKeepRedirectorsButtonToolTipText)
+								.ButtonStyle(FAppStyle::Get(), "FlatButton.Default")
+								.TextStyle(FAppStyle::Get(), "FlatButton.DefaultTextStyle")
+								.OnClicked(this, &SFixupRedirectorsReport::CloseKeepRedirectors)
+							]
+						]
+					]
+				];
+		}
+		else
+		{
+			// Simpler bottom row with just an "Ok" box
+			VerticalBox->AddSlot()
+				.AutoHeight()
+				.Padding(0, 10, 0, 0)
+				[
+					SNew(SBorder)
+					.BorderImage(FAppStyle::GetBrush("NoBorder"))
+					.VAlign(VAlign_Bottom)
+					.HAlign(HAlign_Center)
+					[
+						SAssignNew(KeepButton, SButton)
+						.HAlign(HAlign_Center)
+						.Text(LOCTEXT("Ok", "Ok"))
+						.ButtonStyle(FAppStyle::Get(), "FlatButton.Default")
+						.TextStyle(FAppStyle::Get(), "FlatButton.DefaultTextStyle")
+						.OnClicked(this, &SFixupRedirectorsReport::CloseKeepRedirectors)
+					]
+				];
+		}
+
+		ChildSlot
+		[
+			SNew(SBorder)
+			.BorderImage(FAppStyle::GetBrush("AssetDeleteDialog.Background"))
+			.Padding(10.0f)
+			[
+				VerticalBox.ToSharedRef()	
+			]
+		];
+		// clang-format on 
+	}
+
+private:
+	FReply CloseKeepRedirectors()
+	{
+		ProvideResult(false);
+		return FReply::Handled();
+	}
+	FReply CloseDeleteRedirectors()
+	{
+		ProvideResult(true);
+		return FReply::Handled();
+	}
+	
+	virtual TSharedPtr<SWidget> GetWidgetToFocusOnActivate() override
+	{
+		return KeepButton;
+	}
+	
+	FText GetIntroText() const
+	{
+		if (bAnyProblems)
+		{
+			return LOCTEXT("ReportIntro_SomeProblems", "Some references to redirectors could not be updated and still remain.");
+		}
+		else 
+		{
+			return LOCTEXT("ReportIntro_Problems", "All references to redirectors were successfully updated.");
+		}
+	}
+	
+	FText GetDeleteExplanationText() const
+	{
+		return LOCTEXT("RedirectorDeleteExplanation", 
+			"Delete all redirectors which are no longer referenced locally. " LINE_TERMINATOR 
+			LINE_TERMINATOR
+			"Use with caution as these redirectors may be referenced"
+			" - for example in other source control branches or in other users' workspaces." 
+			);
+	}
+	
+	FText GetKeepExplanationText() const
+	{
+		return LOCTEXT("RedirectorKeepExplanation", 
+			"Keep all redirectors, both referenced and unreferenced." LINE_TERMINATOR
+			LINE_TERMINATOR
+			"This prevents issues caused by references to them in other workspaces or source control branches.");
+	}
+
+	FText GetKeepRedirectorsButtonText() const
+	{
+		return bCanDelete ? LOCTEXT("KeepRedirectors", "Keep Redirectors") : LOCTEXT("Ok", "Ok");
+	}
+
+	FText GetKeepRedirectorsButtonToolTipText() const
+	{
+		return bCanDelete ? LOCTEXT("KeepRedirectorsToolTipText", "Do not delete any redirectors") : FText::GetEmpty();
+	}
+
+	TSharedRef<ITableRow> HandleGenerateTreeViewRow(TSharedPtr<FFixupReportItem> InItem, const TSharedRef<STableViewBase>& InOwnerTable)
+	{
+		return SNew(SFixupRedirectorsReportTableRow, InOwnerTable, InItem);
+	}
+
+	void HandleGetChildren(TSharedPtr<FFixupReportItem> InItem, TArray<TSharedPtr<FFixupReportItem>>& OutChildren)
+	{
+		OutChildren = InItem->Children;
+	}
+
+	TSharedPtr<SButton> KeepButton;
+	TSharedPtr<STreeView<TSharedPtr<FFixupReportItem>>> TreeView;
+	TArray<TSharedPtr<FFixupReportItem>> TreeItems;
+	bool bCanDelete; // Whether we offer the option to delete redirectors 
+	bool bAnyProblems; // Whether there were any redirectors that couldn't be fully fixed up
+};
 
 void FAssetFixUpRedirectors::FixupReferencers(const TArray<UObjectRedirector*>& Objects, const bool bCheckoutDialogPrompt, ERedirectFixupMode FixupMode) const
 {
@@ -77,471 +548,395 @@ void FAssetFixUpRedirectors::ExecuteFixUp(TArray<TWeakObjectPtr<UObjectRedirecto
 	TGuardValue<bool> Guard(bIsFixupReferencersInProgress, true);
 
 	TArray<FRedirectorRefs> RedirectorRefsList;
-	for (auto Object : Objects)
+	for (TWeakObjectPtr<UObjectRedirector> Object : Objects)
 	{
-		auto ObjectRedirector = Object.Get();
-
-		if (ObjectRedirector)
+		if (UObjectRedirector* ObjectRedirector = Object.Get())
 		{
 			RedirectorRefsList.Emplace(ObjectRedirector);
 		}
 	}
 
-	if ( RedirectorRefsList.Num() > 0 )
+	if (RedirectorRefsList.Num() == 0)
 	{
-		// Gather all referencing packages for all redirectors that are being fixed.
-		PopulateRedirectorReferencers(RedirectorRefsList);
-
-		// Update Package Status for all selected redirectors if SCC is enabled
-		if ( UpdatePackageStatus(RedirectorRefsList) )
-		{
-			// Load all referencing packages.
-			TArray<UPackage*> ReferencingPackagesToSave;
-			TArray<UPackage*> LoadedPackages;
-			LoadReferencingPackages(RedirectorRefsList, ReferencingPackagesToSave, LoadedPackages);
-
-			// Add all referencing packages objects that aren't RF_Standalone to the root set to avoid them being GC'd during the following processing
-			TSet<UObject*> RootedObjects;
-			for (UPackage* Package : ReferencingPackagesToSave)
-			{
-				ForEachObjectWithPackage(Package, [&RootedObjects](UObject* Object)
-				{
-					check(!Object->IsRooted());
-					Object->AddToRoot();
-					RootedObjects.Add(Object);
-					return true;
-				}, false, RF_Standalone, EInternalObjectFlags::RootSet);
-			}
-
-			// Check out all referencing packages, leave redirectors for assets referenced by packages that are not checked out and remove those packages from the save list.
-			const bool bUserAcceptedCheckout = CheckOutReferencingPackages(RedirectorRefsList, ReferencingPackagesToSave, bCheckoutDialogPrompt);
-			if ( bUserAcceptedCheckout )
-			{
-				// If any referencing packages are left read-only, the checkout failed or SCC was not enabled. Trim them from the save list and leave redirectors.
-				DetectReadOnlyPackages(RedirectorRefsList, ReferencingPackagesToSave);
-
-				// Fix up referencing FSoftObjectPaths
-				FixUpSoftObjectPaths(RedirectorRefsList, ReferencingPackagesToSave);
-
-				// Save all packages that were referencing any of the assets that were moved without redirectors
-				TArray<UPackage*> FailedToSave;
-				SaveReferencingPackages(ReferencingPackagesToSave, FailedToSave);
-
-				// Save any collections that were referencing any of the redirectors
-				SaveReferencingCollections(RedirectorRefsList);
-
-				// Wait for package referencers to be updated
-				UpdateAssetReferencers(RedirectorRefsList);
-
-				if (FixupMode == ERedirectFixupMode::DeleteFixedUpRedirectors)
-				{
-					// Delete any redirectors that are no longer referenced
-					DeleteRedirectors(RedirectorRefsList, FailedToSave);
-				}
-
-				// Finally, report any failures that happened during the rename
-				ReportFailures(RedirectorRefsList);
-			}
-
-			// Remove objects we manually rooted to allow them to be properly GC'd
-			for (UObject* Object : RootedObjects)
-			{
-				Object->RemoveFromRoot();
-			}
-
-			// If any packages were loaded during the fixup process, make sure we unload them here
-			if (!LoadedPackages.IsEmpty())
-			{
-				for (UPackage* LoadedPackage : LoadedPackages)
-				{
-					ForEachObjectWithPackage(LoadedPackage, [](UObject* Object)
-					{
-						Object->ClearFlags(RF_Standalone);
-						return true;
-					}, false);
-				}
-
-				// Collect garbage.
-				CollectGarbage(GARBAGE_COLLECTION_KEEPFLAGS);
-			}
-		}
+		return;
+	}
+	
+	// Check if we can delete redirectors - if so we need to perform source control operations on them
+	bool bMayDeleteRedirectors = false;
+	switch (FixupMode)
+	{
+		case ERedirectFixupMode::DeleteFixedUpRedirectors:
+			bMayDeleteRedirectors = true;
+			break;
+		case ERedirectFixupMode::PromptForDeletingRedirectors:
+			const UEditorProjectAssetSettings* Settings = GetDefault<UEditorProjectAssetSettings>();
+			bMayDeleteRedirectors = Settings && Settings->bPromptToDeleteUnreferencedRedirectors;
+			break;	
 	}
-}
 
-void FAssetFixUpRedirectors::PopulateRedirectorReferencers(TArray<FRedirectorRefs>& RedirectorsToPopulate) const
-{
+	// Gather all referencing packages for all redirectors that are being fixed.
 	FAssetRegistryModule& AssetRegistryModule = FModuleManager::LoadModuleChecked<FAssetRegistryModule>("AssetRegistry");
-	for ( auto RedirectorRefsIt = RedirectorsToPopulate.CreateIterator(); RedirectorRefsIt; ++RedirectorRefsIt )
+	bool bAnyRefs = false;
+	for (FRedirectorRefs& RedirectorRefs : RedirectorRefsList)
 	{
-		FRedirectorRefs& RedirectorRefs = *RedirectorRefsIt;
 		AssetRegistryModule.Get().GetReferencers(RedirectorRefs.RedirectorPackageName, RedirectorRefs.ReferencingPackageNames);
+		bAnyRefs = bAnyRefs || RedirectorRefs.ReferencingPackageNames.Num() != 0;
+	}
+	
+	if (!bAnyRefs && !bMayDeleteRedirectors)
+	{
+		FMessageDialog::Open(EAppMsgType::Ok, LOCTEXT("NoPackagesToSave", "No referencing assets found to be resaved."));
+		return;
+	}
+	
+	// Build back-reference lookup for later
+	TMultiMap<FName, FRedirectorRefs*> ReferencingAssetToRedirector;
+	for (FRedirectorRefs& RedirectorRefs : RedirectorRefsList)
+	{
+		for (FName PackageName : RedirectorRefs.ReferencingPackageNames)
+		{
+			ReferencingAssetToRedirector.Add(PackageName, &RedirectorRefs);	
+		}
 	}
-}
 
-bool FAssetFixUpRedirectors::UpdatePackageStatus(const TArray<FRedirectorRefs>& RedirectorsToFix) const
-{
-	ISourceControlProvider& SourceControlProvider = ISourceControlModule::Get().GetProvider();
-	if ( ISourceControlModule::Get().IsEnabled() )
+	// Update Package Status for all selected redirectors if SCC is enabled and we may delete
+	if (bMayDeleteRedirectors && ISourceControlModule::Get().IsEnabled())
 	{
+		ISourceControlProvider& SourceControlProvider = ISourceControlModule::Get().GetProvider();
 		// Update the source control server availability to make sure we can do the rename operation
 		SourceControlProvider.Login();
-		if ( !SourceControlProvider.IsAvailable() )
+		if (!SourceControlProvider.IsAvailable())
 		{
 			// We have failed to update source control even though it is enabled. This is critical and we can not continue
 			FMessageDialog::Open( EAppMsgType::Ok, NSLOCTEXT("UnrealEd", "SourceControl_ServerUnresponsive", "Revision Control is unresponsive. Please check your connection and try again.") );
-			return false;
+			return;
 		}
 
 		TArray<UPackage*> PackagesToAddToSCCUpdate;
-
-		for ( auto RedirectorRefsIt = RedirectorsToFix.CreateConstIterator(); RedirectorRefsIt; ++RedirectorRefsIt )
+		for (const FRedirectorRefs& RedirectorRefs : RedirectorRefsList)
 		{
-			const FRedirectorRefs& RedirectorRefs = *RedirectorRefsIt;
 			PackagesToAddToSCCUpdate.Add(RedirectorRefs.Redirector->GetOutermost());
 		}
-
 		SourceControlProvider.Execute(ISourceControlOperation::Create<FUpdateStatus>(), PackagesToAddToSCCUpdate);
 	}
 
-	return true;
-}
-
-void FAssetFixUpRedirectors::LoadReferencingPackages(TArray<FRedirectorRefs>& RedirectorsToFix, TArray<UPackage*>& OutReferencingPackagesToSave, TArray<UPackage*>& OutLoadedPackages) const
-{
-	FScopedSlowTask SlowTask( static_cast<float>(RedirectorsToFix.Num()), LOCTEXT( "LoadingReferencingPackages", "Loading Referencing Packages..." ) );
-	SlowTask.MakeDialog();
-
-	ISourceControlProvider& SourceControlProvider = ISourceControlModule::Get().GetProvider();
-
-	// Load all packages that reference each redirector, if possible
-	for ( auto RedirectorRefsIt = RedirectorsToFix.CreateIterator(); RedirectorRefsIt; ++RedirectorRefsIt )
+	// Load all referencing packages.
+	TSet<UPackage*> ReferencingPackagesToSave;
+	TSet<UPackage*> LoadedPackages;
 	{
-		SlowTask.EnterProgressFrame(1);
+		FScopedSlowTask SlowTask(static_cast<float>(RedirectorRefsList.Num()), LOCTEXT( "LoadingReferencingPackages", "Loading Referencing Packages..." ) );
+		SlowTask.MakeDialog();
+		ISourceControlProvider& SourceControlProvider = ISourceControlModule::Get().GetProvider();
 
-		FRedirectorRefs& RedirectorRefs = *RedirectorRefsIt;
-		if ( ISourceControlModule::Get().IsEnabled() )
+		// Load all packages that reference each redirector, if possible
+		for (FRedirectorRefs& RedirectorRefs : RedirectorRefsList)
 		{
-			FSourceControlStatePtr SourceControlState = SourceControlProvider.GetState(RedirectorRefs.Redirector->GetOutermost(), EStateCacheUsage::Use);
-			const bool bValidSCCState = !SourceControlState.IsValid() || SourceControlState->IsAdded() || SourceControlState->IsCheckedOut() || SourceControlState->CanCheckout() || !SourceControlState->IsSourceControlled() || SourceControlState->IsIgnored();
-
-			if ( !bValidSCCState )
+			SlowTask.EnterProgressFrame(1);
+			if (bMayDeleteRedirectors && ISourceControlModule::Get().IsEnabled())
 			{
-				RedirectorRefs.bRedirectorValidForFixup = false;
-				RedirectorRefs.FailureReason = LOCTEXT("RedirectorFixupFailed_BadSCC", "Redirector could not be checked out or marked for delete");
-			}
-		}
+				FSourceControlStatePtr SourceControlState = SourceControlProvider.GetState(RedirectorRefs.Redirector->GetOutermost(), EStateCacheUsage::Use);
+				const bool bValidSCCState = !SourceControlState.IsValid() || SourceControlState->IsAdded() || SourceControlState->IsCheckedOut() || SourceControlState->CanCheckout() || !SourceControlState->IsSourceControlled() || SourceControlState->IsIgnored();
 
-		// Load all referencers
-		for ( auto PackageNameIt = RedirectorRefs.ReferencingPackageNames.CreateConstIterator(); PackageNameIt; ++PackageNameIt )
-		{
-			const FString PackageName = (*PackageNameIt).ToString();
-
-			// Find the package in memory. If it is not in memory, try to load it
-			UPackage* Package = FindPackage(NULL, *PackageName);
-			if ( !Package )
-			{
-				Package = LoadPackage(NULL, *PackageName, LOAD_None);
-				if (Package)
+				if (!bValidSCCState)
 				{
-					OutLoadedPackages.Add(Package);
+					RedirectorRefs.bSCCError = true;
 				}
 			}
 
-			if ( Package )
+			// Load all referencers
+			for (FName ReferencingPackageName : RedirectorRefs.ReferencingPackageNames)
 			{
-				if ( Package->HasAnyPackageFlags(PKG_CompiledIn) )
+				FNameBuilder PackageName{ReferencingPackageName};
+
+				// Find the package in memory. If it is not in memory, try to load it
+				UPackage* Package = FindPackage(nullptr, *PackageName);
+				if (!Package)
 				{
-					// This is a script reference
-					RedirectorRefs.bRedirectorValidForFixup = false;
-					RedirectorRefs.FailureReason = FText::Format(LOCTEXT("RedirectorFixupFailed_CodeReference", "Redirector is referenced by code. Package: {0}"), FText::FromString(PackageName));
+					Package = LoadPackage(nullptr, *PackageName, LOAD_None);
+					if (Package)
+					{
+						LoadedPackages.Add(Package);
+					}
 				}
-				else
+
+				if (Package)
 				{
-					// If we found a valid package, mark it for save
-					OutReferencingPackagesToSave.AddUnique(Package);
+					if (Package->HasAnyPackageFlags(PKG_CompiledIn))
+					{
+						// This is a script reference
+						RedirectorRefs.OtherFailures.Add(FText::Format(LOCTEXT("RedirectorFixupFailed_CodeReference", "Redirector is referenced by code. Package: {0}"), FText::FromName(ReferencingPackageName)));
+					}
+					else
+					{
+						// If we found a valid package, mark it for save
+						ReferencingPackagesToSave.Add(Package);
+					}
 				}
 			}
 		}
 	}
-}
 
-bool FAssetFixUpRedirectors::CheckOutReferencingPackages(TArray<FRedirectorRefs>& RedirectorsToFix, TArray<UPackage*>& InOutReferencingPackagesToSave, const bool bCheckoutDialogPrompt) const
-{
-	// Prompt to check out all successfully loaded packages
-	bool bUserAcceptedCheckout = true;
+	ON_SCOPE_EXIT {
+		// If any packages were loaded during the fixup process, make sure we unload them here
+		if (!LoadedPackages.IsEmpty())
+		{
+			FText ErrorMessage;
+			UPackageTools::UnloadPackages(LoadedPackages.Array(), ErrorMessage, true);
+			if (!ErrorMessage.IsEmpty())
+			{
+				FTextBuilder Builder;
+				Builder.AppendLine(LOCTEXT("ErrorsUnloadingPackages", "Errors were encountered unloading packages which were loaded to update redirector references. Some assets may still be loaded. "));
+				Builder.AppendLine();
+				Builder.AppendLine(ErrorMessage);
+				FMessageDialog::Open(EAppMsgType::Ok, Builder.ToText());
+			}
+		}
+	};
 
-	if ( InOutReferencingPackagesToSave.Num() > 0 )
+	// Add all referencing packages objects that aren't RF_Standalone to the root set to avoid them being GC'd during the following processing
+	TArray<TStrongObjectPtr<UObject>> RootedObjects;
+	for (UPackage* Package : ReferencingPackagesToSave)
+	{
+		ForEachObjectWithPackage(Package, [&RootedObjects](UObject* Object)
+		{
+			RootedObjects.Emplace(Object);
+			return true;
+		}, false, RF_Standalone, EInternalObjectFlags::RootSet);
+	}
+	
+	// Check out all referencing packages, leave redirectors for assets referenced by packages that are not checked out and remove those packages from the save list.
+	bool bUserAcceptedCheckout = true; // If source control is disabled, assume checkout was selected
+	if (ISourceControlModule::Get().IsEnabled() && ReferencingPackagesToSave.Num() > 0)
 	{
-		if ( ISourceControlModule::Get().IsEnabled() )
+		TArray<UPackage*> PackagesCheckedOutOrMadeWritable;
+		TArray<UPackage*> PackagesNotNeedingCheckout;
+		if (bCheckoutDialogPrompt)
 		{
-			TArray<UPackage*> PackagesCheckedOutOrMadeWritable;
-			TArray<UPackage*> PackagesNotNeedingCheckout;
-			if (bCheckoutDialogPrompt)
+			bUserAcceptedCheckout = FEditorFileUtils::PromptToCheckoutPackages(false, ReferencingPackagesToSave.Array(), &PackagesCheckedOutOrMadeWritable, &PackagesNotNeedingCheckout);
+		}
+		else
+		{
+			const bool bErrorIfAlreadyCheckedOut = false;
+			const bool bConfirmPackageBranchCheckOutStatus = false;
+			FEditorFileUtils::CheckoutPackages(ReferencingPackagesToSave.Array(), &PackagesCheckedOutOrMadeWritable, bErrorIfAlreadyCheckedOut, bConfirmPackageBranchCheckOutStatus);
+		}
+
+		if (bUserAcceptedCheckout)
+		{
+			TSet<UPackage*> PackagesThatCouldNotBeCheckedOut = ReferencingPackagesToSave;
+			for (UPackage* Package : PackagesCheckedOutOrMadeWritable)
 			{
-				bUserAcceptedCheckout = FEditorFileUtils::PromptToCheckoutPackages( false, InOutReferencingPackagesToSave, &PackagesCheckedOutOrMadeWritable, &PackagesNotNeedingCheckout );
+				PackagesThatCouldNotBeCheckedOut.Remove(Package);
 			}
-			else
+			for (UPackage* Package : PackagesNotNeedingCheckout)
 			{
-				const bool bErrorIfAlreadyCheckedOut = false;
-				const bool bConfirmPackageBranchCheckOutStatus = false;
-				FEditorFileUtils::CheckoutPackages(InOutReferencingPackagesToSave, &PackagesCheckedOutOrMadeWritable, bErrorIfAlreadyCheckedOut, bConfirmPackageBranchCheckOutStatus);
+				PackagesThatCouldNotBeCheckedOut.Remove(Package);
 			}
 
-			if ( bUserAcceptedCheckout )
+			for (UPackage* Package : PackagesThatCouldNotBeCheckedOut)
 			{
-				TArray<UPackage*> PackagesThatCouldNotBeCheckedOut = InOutReferencingPackagesToSave;
-
-				for ( auto PackageIt = PackagesCheckedOutOrMadeWritable.CreateConstIterator(); PackageIt; ++PackageIt )
-				{
-					PackagesThatCouldNotBeCheckedOut.Remove(*PackageIt);
-				}
+				FName PackageName = Package->GetFName(); // Key iterator requires we copy this
 
-				for ( auto PackageIt = PackagesNotNeedingCheckout.CreateConstIterator(); PackageIt; ++PackageIt )
+				// Note which redirectors will still be required because this package could not be checked out 
+				for (auto It = ReferencingAssetToRedirector.CreateKeyIterator(PackageName); It; ++It)
 				{
-					PackagesThatCouldNotBeCheckedOut.Remove(*PackageIt);
+					It.Value()->LockedReferencerPackageNames.Add(Package->GetFName());
 				}
 
-				for ( auto PackageIt = PackagesThatCouldNotBeCheckedOut.CreateConstIterator(); PackageIt; ++PackageIt )
-				{
-					const FName NonCheckedOutPackageName = (*PackageIt)->GetFName();
-
-					for ( auto RedirectorRefsIt = RedirectorsToFix.CreateIterator(); RedirectorRefsIt; ++RedirectorRefsIt )
-					{
-						FRedirectorRefs& RedirectorRefs = *RedirectorRefsIt;
-						if ( RedirectorRefs.ReferencingPackageNames.Contains(NonCheckedOutPackageName) )
-						{
-							// We did not check out at least one of the packages we needed to. This redirector can not be fixed up.
-							RedirectorRefs.FailureReason = FText::Format(LOCTEXT("RedirectorFixupFailed_NotCheckedOut", "Referencing package {0} was not checked out"), FText::FromName(NonCheckedOutPackageName));
-							RedirectorRefs.bRedirectorValidForFixup = false;
-						}
-					}
-
-					InOutReferencingPackagesToSave.Remove(*PackageIt);
-				}
+				// Don't save anything that wasn't checked out 
+				ReferencingPackagesToSave.Remove(Package);
 			}
 		}
 	}
 
-	return bUserAcceptedCheckout;
-}
-
-void FAssetFixUpRedirectors::DetectReadOnlyPackages(TArray<FRedirectorRefs>& RedirectorsToFix, TArray<UPackage*>& InOutReferencingPackagesToSave) const
-{
-	// For each valid package...
-	for ( int32 PackageIdx = InOutReferencingPackagesToSave.Num() - 1; PackageIdx >= 0; --PackageIdx )
+	if (bUserAcceptedCheckout)
 	{
-		UPackage* Package = InOutReferencingPackagesToSave[PackageIdx];
-
-		if ( Package )
+		// If any referencing packages are left read-only, the checkout failed or SCC was not enabled. Trim them from the save list and leave redirectors.
+		for (auto It = ReferencingPackagesToSave.CreateIterator(); It; ++It)
 		{
-			// Find the package filename
+			UPackage* Package = *It;
+			if (!ensure(Package))
+			{
+				It.RemoveCurrent();
+				continue;
+			}
+			
+			// If the file is read only
 			FString Filename;
-			if ( FPackageName::DoesPackageExist(Package->GetName(), &Filename) )
+			if (FPackageName::DoesPackageExist(Package->GetName(), &Filename) 
+			&& IFileManager::Get().IsReadOnly(*Filename))
 			{
-				// If the file is read only
-				if ( IFileManager::Get().IsReadOnly(*Filename) )
+				// Note which redirectors will still be required because this package was read only
+				FName PackageName = Package->GetFName(); // Key iterator requires we copy this
+				for (auto RedirectorIt = ReferencingAssetToRedirector.CreateKeyIterator(PackageName); RedirectorIt; ++RedirectorIt)
 				{
-					FName PackageName = Package->GetFName();
-
-					// Find all assets that were referenced by this package to create a redirector when named
-					for ( auto RedirectorIt = RedirectorsToFix.CreateIterator(); RedirectorIt; ++RedirectorIt )
-					{
-						FRedirectorRefs& RedirectorRefs = *RedirectorIt;
-						if ( RedirectorRefs.ReferencingPackageNames.Contains(PackageName) )
-						{
-							RedirectorRefs.FailureReason = FText::Format(LOCTEXT("RedirectorFixupFailed_ReadOnly", "Referencing package {0} was read-only"), FText::FromName(PackageName));
-							RedirectorRefs.bRedirectorValidForFixup = false;
-						}
-					}
-
-					// Remove the package from the save list
-					InOutReferencingPackagesToSave.RemoveAt(PackageIdx);
+					RedirectorIt.Value()->LockedReferencerPackageNames.Add(Package->GetFName());
 				}
+
+				// Remove the package from the save list
+				It.RemoveCurrent();
 			}
 		}
-	}
-}
 
-void FAssetFixUpRedirectors::SaveReferencingPackages(const TArray<UPackage*>& ReferencingPackagesToSave, TArray<UPackage*>& OutFailedToSave) const
-{
-	if ( ReferencingPackagesToSave.Num() > 0 )
-	{
-		// Get the list of filenames before calling save because some of the saved packages can get GCed if they are empty packages
-		const TArray<FString> Filenames = USourceControlHelpers::PackageFilenames(ReferencingPackagesToSave);
-
-		const bool bCheckDirty = false;
-		const bool bPromptToSave = false;
-		FEditorFileUtils::PromptForCheckoutAndSave(ReferencingPackagesToSave, bCheckDirty, bPromptToSave, &OutFailedToSave);
-
-		ISourceControlModule::Get().QueueStatusUpdate(Filenames);
-	}
-}
+		// Fix up referencing FSoftObjectPaths
+		{
+			TSet<UPackage*> PackagesToCheck = ReferencingPackagesToSave;
 
-void FAssetFixUpRedirectors::SaveReferencingCollections(TArray<FRedirectorRefs>& RedirectorsToFix) const
-{
-	FAssetRegistryModule& AssetRegistryModule = FModuleManager::LoadModuleChecked<FAssetRegistryModule>(TEXT("AssetRegistry"));
-	FCollectionManagerModule& CollectionManagerModule = FCollectionManagerModule::GetModule();
+			TArray<UPackage*> Tmp;
+			FEditorFileUtils::GetDirtyWorldPackages(Tmp);
+			FEditorFileUtils::GetDirtyContentPackages(Tmp);
+			PackagesToCheck.Append(Tmp);
 
-	// Find all collections that were referenced by any of the redirectors that are potentially going to be removed and attempt to re-save them
-	// The redirectors themselves will have already been fixed up, as collections do that once the asset registry has been populated, 
-	// however collections lazily re-save redirector fix-up to avoid SCC issues, so we need to force that now
-	for (FRedirectorRefs& RedirectorRefs : RedirectorsToFix)
-	{
-		// Follow each link in the redirector, and notify the collections manager that it is going to be removed - this will force it to re-save any required collections
-		for (UObjectRedirector* Redirector = RedirectorRefs.Redirector; Redirector; Redirector = Cast<UObjectRedirector>(Redirector->DestinationObject))
-		{
-			const FSoftObjectPath RedirectorObjectPath = FSoftObjectPath(Redirector);
-			if (!CollectionManagerModule.Get().HandleRedirectorDeleted(RedirectorObjectPath))
+			TMap<FSoftObjectPath, FSoftObjectPath> RedirectorMap;
+			for (const FRedirectorRefs& RedirectorRef : RedirectorRefsList)
 			{
-				RedirectorRefs.FailureReason = FText::Format(LOCTEXT("RedirectorFixupFailed_CollectionsFailedToSave", "Referencing collection(s) failed to save: {0}"), CollectionManagerModule.Get().GetLastError());
-				RedirectorRefs.bRedirectorValidForFixup = false;
-			}
-		}
-	}
-}
-
-void FAssetFixUpRedirectors::UpdateAssetReferencers(const TArray<FRedirectorRefs>& RedirectorsToFix) const
-{
-	// Load the asset registry module
-	FAssetRegistryModule& AssetRegistryModule = FModuleManager::LoadModuleChecked<FAssetRegistryModule>(TEXT("AssetRegistry"));
+				const UObjectRedirector* Redirector = RedirectorRef.Redirector.Get();
+				FSoftObjectPath OldPath = FSoftObjectPath(Redirector);
+				FSoftObjectPath NewPath = FSoftObjectPath(Redirector->DestinationObject);
 
-	TArray<FString> AssetPaths;
-	for (const auto& Redirector : RedirectorsToFix)
-	{
-		AssetPaths.AddUnique(FPackageName::GetLongPackagePath(Redirector.RedirectorPackageName.ToString()) / TEXT("")); // Ensure trailing slash
+				RedirectorMap.Add(OldPath, NewPath);
+				if (UBlueprint* Blueprint = Cast<UBlueprint>(Redirector->DestinationObject))
+				{
+					// Add redirect for class and default as well
+					RedirectorMap.Add(FString::Printf(TEXT("%s_C"), *OldPath.ToString()), FString::Printf(TEXT("%s_C"), *NewPath.ToString()));
+					RedirectorMap.Add(FString::Printf(TEXT("%s.Default__%s_C"), *OldPath.GetLongPackageName(), *OldPath.GetAssetName()), FString::Printf(TEXT("%s.Default__%s_C"), *NewPath.GetLongPackageName(), *NewPath.GetAssetName()));
+				}
+			}
 
-		for (const auto& Referencer : Redirector.ReferencingPackageNames)
-		{
-			AssetPaths.AddUnique(FPackageName::GetLongPackagePath(Referencer.ToString()) / TEXT("")); // Ensure trailing slash
+			UAssetToolsImpl::Get().AssetRenameManager->RenameReferencingSoftObjectPaths(PackagesToCheck.Array(), RedirectorMap);
 		}
-	}
-	AssetRegistryModule.Get().ScanPathsSynchronous(AssetPaths, true);
-}
 
-void FAssetFixUpRedirectors::DeleteRedirectors(TArray<FRedirectorRefs>& RedirectorsToFix, const TArray<UPackage*>& FailedToSave) const
-{
-	TArray<UObject*> ObjectsToDelete;
-	for ( auto RedirectorIt = RedirectorsToFix.CreateIterator(); RedirectorIt; ++RedirectorIt )
-	{
-		FRedirectorRefs& RedirectorRefs = *RedirectorIt;
-		if ( RedirectorRefs.bRedirectorValidForFixup )
+		// Save all packages that were referencing any of the assets that were moved without redirectors
+		TArray<UPackage*> FailedToSave;
+		if (ReferencingPackagesToSave.Num() > 0)
 		{
-			check(RedirectorRefs.Redirector);
+			// Get the list of filenames before calling save because some of the saved packages can get GCed if they are empty packages
+			const TArray<FString> Filenames = USourceControlHelpers::PackageFilenames(ReferencingPackagesToSave.Array());
 
-			bool bAllReferencersFixedUp = true;
-			for (const auto& ReferencingPackageName : RedirectorRefs.ReferencingPackageNames)
+			const bool bCheckDirty = false;
+			const bool bPromptToSave = false;
+			FEditorFileUtils::PromptForCheckoutAndSave(ReferencingPackagesToSave.Array(), bCheckDirty, bPromptToSave, &FailedToSave);
+			for (UPackage* Package : FailedToSave)
 			{
-				if (FailedToSave.ContainsByPredicate([&](UPackage* Package) { return Package->GetFName() == ReferencingPackageName; }))
+				FName PackageName = Package->GetFName(); // Key iterator requires we copy this
+				for (auto It = ReferencingAssetToRedirector.CreateKeyIterator(PackageName); It; ++It)	
 				{
-					bAllReferencersFixedUp = false;
-					break;
+					It.Value()->FailedReferencerPackageNames.Add(Package->GetFName());
 				}
 			}
 
-			if (!bAllReferencersFixedUp)
-			{
-				continue;
-			}
+			ISourceControlModule::Get().QueueStatusUpdate(Filenames);
+		}
 
-			// Add all redirectors found in this package to the redirectors to delete list.
-			// All redirectors in this package should be fixed up.
-			UPackage* RedirectorPackage = RedirectorRefs.Redirector->GetOutermost();
-			TArray<UObject*> AssetsInRedirectorPackage;
-			GetObjectsWithOuter(RedirectorPackage, AssetsInRedirectorPackage, /*bIncludeNestedObjects=*/false);
-			UMetaData* PackageMetaData = NULL;
-			bool bContainsAtLeastOneOtherAsset = false;
-			for ( auto ObjIt = AssetsInRedirectorPackage.CreateConstIterator(); ObjIt; ++ObjIt )
+		// Save any collections that were referencing any of the redirectors
+		FCollectionManagerModule& CollectionManagerModule = FCollectionManagerModule::GetModule();
+
+		// Find all collections that were referenced by any of the redirectors that are potentially going to be removed and attempt to re-save them
+		// The redirectors themselves will have already been fixed up, as collections do that once the asset registry has been populated, 
+		// however collections lazily re-save redirector fix-up to avoid SCC issues, so we need to force that now
+		for (FRedirectorRefs& RedirectorRefs : RedirectorRefsList)
+		{
+			// Follow each link in the redirector, and notify the collections manager that it is going to be removed - this will force it to re-save any required collections
+			for (const UObjectRedirector* Redirector = RedirectorRefs.Redirector.Get(); Redirector; Redirector = Cast<UObjectRedirector>(Redirector->DestinationObject))
 			{
-				if ( UObjectRedirector* Redirector = Cast<UObjectRedirector>(*ObjIt) )
+				const FSoftObjectPath RedirectorObjectPath = FSoftObjectPath(Redirector);
+				if (!CollectionManagerModule.Get().HandleRedirectorDeleted(RedirectorObjectPath))
 				{
-					Redirector->RemoveFromRoot();
-					ObjectsToDelete.Add(Redirector);
-				}
-				else if ( UMetaData* MetaData = Cast<UMetaData>(*ObjIt) )
-				{
-					PackageMetaData = MetaData;
-				}
-				else
-				{
-					bContainsAtLeastOneOtherAsset = true;
+					RedirectorRefs.OtherFailures.Add(FText::Format(LOCTEXT("RedirectorFixupFailed_CollectionsFailedToSave", "Referencing collection(s) failed to save: {0}"), CollectionManagerModule.Get().GetLastError()));
 				}
 			}
+		}
 
-			if ( !bContainsAtLeastOneOtherAsset )
+		// Wait for package referencers to be updated
+		// Load the asset registry module
+		{
+			TArray<FString> AssetPaths;
+			for (const FRedirectorRefs& Redirector : RedirectorRefsList)
 			{
-				RedirectorPackage->RemoveFromRoot();
-				ObjectsToDelete.Add(RedirectorPackage);
+				AssetPaths.AddUnique(FPackageName::GetLongPackagePath(Redirector.RedirectorPackageName.ToString()) / TEXT("")); // Ensure trailing slash
 
-				// @todo we shouldnt be worrying about metadata objects here, ObjectTools::CleanUpAfterSuccessfulDelete should
-				if ( PackageMetaData )
+				for (const auto& Referencer : Redirector.ReferencingPackageNames)
 				{
-					PackageMetaData->RemoveFromRoot();
-					ObjectsToDelete.Add(PackageMetaData);
+					AssetPaths.AddUnique(FPackageName::GetLongPackagePath(Referencer.ToString()) / TEXT("")); // Ensure trailing slash
 				}
 			}
-
-			// This redirector will be deleted, NULL the reference here
-			RedirectorRefs.Redirector = NULL;
+			AssetRegistryModule.Get().ScanPathsSynchronous(AssetPaths, true);
 		}
-	}
 
-	if ( ObjectsToDelete.Num() > 0 )
-	{
-		ObjectTools::DeleteObjects(ObjectsToDelete, false);
-	}
-}
-
-void FAssetFixUpRedirectors::ReportFailures(const TArray<FRedirectorRefs>& RedirectorsToFix) const
-{
-	FMessageLog EditorErrors("EditorErrors");
-	bool bTitleOutput = false;
-
-	for ( auto RedirectorIt = RedirectorsToFix.CreateConstIterator(); RedirectorIt; ++RedirectorIt )
-	{
-		const FRedirectorRefs& RedirectorRefs = *RedirectorIt;
-		if ( !RedirectorRefs.bRedirectorValidForFixup )
+		// Show user report and check whether we should delete
+		bool bCanDelete = false;
+		switch (FixupMode)
 		{
-			if(!bTitleOutput)
-			{
-				EditorErrors.Info(LOCTEXT("RedirectorFixupFailedMessage", "The following redirectors could not be completely fixed up"));
-				bTitleOutput = true;
-			}
-			FFormatNamedArguments Arguments;
-			Arguments.Add(TEXT("PackageName"), FText::FromName(RedirectorRefs.RedirectorPackageName));
-			Arguments.Add(TEXT("FailureReason"), FText::FromString(RedirectorRefs.FailureReason.ToString()));
-			EditorErrors.Warning(FText::Format(LOCTEXT("RedirectorFixupFailedReason", "{PackageName} - {FailureReason}"), Arguments ));
+			case ERedirectFixupMode::DeleteFixedUpRedirectors:
+				bCanDelete = true;
+				break;
+			case ERedirectFixupMode::PromptForDeletingRedirectors:
+				const UEditorProjectAssetSettings* Settings = GetDefault<UEditorProjectAssetSettings>();
+				bCanDelete = Settings && Settings->bPromptToDeleteUnreferencedRedirectors;
+				break;
 		}
-	}
+		
 
-	EditorErrors.Open();
-}
-
-void FAssetFixUpRedirectors::FixUpSoftObjectPaths(const TArray<FRedirectorRefs>& RedirectorsToFix, const TArray<UPackage*>& InReferencingPackagesToSave) const
-{
-	TArray<UPackage *> PackagesToCheck(InReferencingPackagesToSave);
-
-	FEditorFileUtils::GetDirtyWorldPackages(PackagesToCheck);
-	FEditorFileUtils::GetDirtyContentPackages(PackagesToCheck);
+		TSharedPtr<SFixupRedirectorsReport> Dialog = SNew(SFixupRedirectorsReport, RedirectorRefsList, bCanDelete);
+		const bool bDoDelete = Dialog->ShowModalDialog(LOCTEXT("RedirectorUpdateReport", "Redirector Update Report"));
+		if (!bDoDelete)
+		{
+			return;
+		}
 
-	TMap<FSoftObjectPath, FSoftObjectPath> RedirectorMap;
+		TArray<UObject*> ObjectsToDelete;
+		for (FRedirectorRefs& RedirectorRefs : RedirectorRefsList)
+		{
+			if (RedirectorRefs.OtherFailures.IsEmpty() 
+			&&  RedirectorRefs.LockedReferencerPackageNames.IsEmpty()
+			&& 	RedirectorRefs.FailedReferencerPackageNames.IsEmpty())
+			{
+				ensure(RedirectorRefs.Redirector);
+				// Add all redirectors found in this package to the redirectors to delete list.
+				// All redirectors in this package should be fixed up.
+				UPackage* RedirectorPackage = RedirectorRefs.Redirector->GetOutermost();
+				UMetaData* PackageMetaData = nullptr;
+				bool bContainsAtLeastOneOtherAsset = false;
+				ForEachObjectWithOuter(RedirectorPackage, [&PackageMetaData, &ObjectsToDelete, &bContainsAtLeastOneOtherAsset](UObject* Obj)
+				{
+					if (UObjectRedirector* Redirector = Cast<UObjectRedirector>(Obj))
+					{
+						Redirector->RemoveFromRoot();
+						ObjectsToDelete.Add(Redirector);
+					}
+					else if (UMetaData* MetaData = Cast<UMetaData>(Obj))
+					{
+						PackageMetaData = MetaData;
+					}
+					else
+					{
+						bContainsAtLeastOneOtherAsset = true;
+					}
+				});
 
-	for (const FRedirectorRefs& RedirectorRef : RedirectorsToFix)
-	{
-		UObjectRedirector* Redirector = RedirectorRef.Redirector;
-		FSoftObjectPath OldPath = FSoftObjectPath(Redirector);
-		FSoftObjectPath NewPath = FSoftObjectPath(Redirector->DestinationObject);
+				if ( !bContainsAtLeastOneOtherAsset )
+				{
+					RedirectorPackage->RemoveFromRoot();
+					ObjectsToDelete.Add(RedirectorPackage);
 
-		RedirectorMap.Add(OldPath, NewPath);
+					// @todo we shouldnt be worrying about metadata objects here, ObjectTools::CleanUpAfterSuccessfulDelete should
+					if ( PackageMetaData )
+					{
+						PackageMetaData->RemoveFromRoot();
+						ObjectsToDelete.Add(PackageMetaData);
+					}
+				}
+			}
+		}
+		
+		// Release all redirector references before executing delete operation
+		RedirectorRefsList.Empty();
 
-		if (UBlueprint* Blueprint = Cast<UBlueprint>(Redirector->DestinationObject))
+		if ( ObjectsToDelete.Num() > 0 )
 		{
-			// Add redirect for class and default as well
-			RedirectorMap.Add(FString::Printf(TEXT("%s_C"), *OldPath.ToString()), FString::Printf(TEXT("%s_C"), *NewPath.ToString()));
-			RedirectorMap.Add(FString::Printf(TEXT("%s.Default__%s_C"), *OldPath.GetLongPackageName(), *OldPath.GetAssetName()), FString::Printf(TEXT("%s.Default__%s_C"), *NewPath.GetLongPackageName(), *NewPath.GetAssetName()));
+			ObjectTools::DeleteObjects(ObjectsToDelete, false);
 		}
 	}
-
-	UAssetToolsImpl::Get().AssetRenameManager->RenameReferencingSoftObjectPaths(PackagesToCheck, RedirectorMap);
 }
 
 #undef LOCTEXT_NAMESPACE
diff --git a/Engine/Source/Developer/AssetTools/Private/AssetFixUpRedirectors.h b/Engine/Source/Developer/AssetTools/Private/AssetFixUpRedirectors.h
--- a/Engine/Source/Developer/AssetTools/Private/AssetFixUpRedirectors.h
+++ b/Engine/Source/Developer/AssetTools/Private/AssetFixUpRedirectors.h
@@ -24,46 +24,6 @@ private:
 	/** The core code of the fixup operation */
 	void ExecuteFixUp(TArray<TWeakObjectPtr<UObjectRedirector>> Objects, bool bCheckoutDialogPrompt, ERedirectFixupMode FixupMode) const;
 
-	/** Fills out the Referencing packages for all the redirectors described in AssetsToPopulate */
-	void PopulateRedirectorReferencers(TArray<FRedirectorRefs>& RedirectorsToPopulate) const;
-
-	/** Updates the source control status of the packages containing the assets to rename */
-	bool UpdatePackageStatus(const TArray<FRedirectorRefs>& RedirectorsToFix) const;
-
-	/** 
-	  * Loads all referencing packages to redirectors in RedirectorsToFix, finds redirectors whose references can
-	  * not be fixed up, and returns a list of referencing packages to save.
-	  */
-	void LoadReferencingPackages(TArray<FRedirectorRefs>& RedirectorsToFix, TArray<UPackage*>& OutReferencingPackagesToSave, TArray<UPackage*>& OutLoadedPackages)const;
-
-	/** 
-	  * Check out referencing packages and marks assets whose referencing packages were not checked out to not fix the redirector.
-	  * Trims PackagesToSave when necessary.
-	  * Returns true if the user opted to continue the operation or no dialog was required.
-	  */
-	bool CheckOutReferencingPackages(TArray<FRedirectorRefs>& RedirectorsToFix, TArray<UPackage*>& InOutReferencingPackagesToSave, bool bCheckoutDialogPrompt) const;
-
-	/** Finds any read only packages and removes them from the save list. Redirectors referenced by these packages will not be fixed up. */ 
-	void DetectReadOnlyPackages(TArray<FRedirectorRefs>& RedirectorsToFix, TArray<UPackage*>& InOutReferencingPackagesToSave) const;
-
-	/** FixUp soft object paths */
-	void FixUpSoftObjectPaths(const TArray<FRedirectorRefs>& RedirectorsToFix, const TArray<UPackage*>& InReferencingPackagesToSave) const;
-
-	/** Saves all the referencing packages and updates SCC state */
-	void SaveReferencingPackages(const TArray<UPackage*>& ReferencingPackagesToSave, TArray<UPackage*>& OutFailedToSave) const;
-
-	/** Saves any collections that were referencing any of the redirectors and updates SCC state */
-	void SaveReferencingCollections(TArray<FRedirectorRefs>& RedirectorsToFix) const;
-
-	/** Waits for the asset registry to update its asset referencer cache */
-	void UpdateAssetReferencers(const TArray<FRedirectorRefs>& RedirectorsToFix) const;
-
-	/** Deletes redirectors that are valid to delete */
-	void DeleteRedirectors(TArray<FRedirectorRefs>& RedirectorsToFix, const TArray<UPackage*>& FailedToSave) const;
-
-	/** Report any failures that may have happened during the rename */
-	void ReportFailures(const TArray<FRedirectorRefs>& RedirectorsToFix) const;
-
 private:
 
 	mutable bool bIsFixupReferencersInProgress = false;
diff --git a/Engine/Source/Developer/AssetTools/Public/IAssetTools.h b/Engine/Source/Developer/AssetTools/Public/IAssetTools.h
--- a/Engine/Source/Developer/AssetTools/Public/IAssetTools.h
+++ b/Engine/Source/Developer/AssetTools/Public/IAssetTools.h
@@ -71,7 +71,11 @@ enum class ERedirectFixupMode
 	DeleteFixedUpRedirectors,
 
 	// Leave the redirectors around even if no longer locally referenced
-	LeaveFixedUpRedirectors
+	LeaveFixedUpRedirectors,
+
+	// If permitted, prompt the user to delete any redirectors that are no longer locally referenced.
+	// Projects may disable deletion if it interferes with their workflow/branch setup/etc.
+	PromptForDeletingRedirectors,
 };
 
 USTRUCT(BlueprintType)
diff --git a/Engine/Source/Editor/UnrealEd/Public/Settings/EditorProjectSettings.h b/Engine/Source/Editor/UnrealEd/Public/Settings/EditorProjectSettings.h
--- a/Engine/Source/Editor/UnrealEd/Public/Settings/EditorProjectSettings.h
+++ b/Engine/Source/Editor/UnrealEd/Public/Settings/EditorProjectSettings.h
@@ -270,8 +270,24 @@ public:
 
 };
 
+/**
+ * Settings for how developers interact with assets. Stored in default config, per-project
+ */
+UCLASS(config = Editor, defaultconfig, meta = (DisplayName = "Assets"), MinimalAPI)
+class UEditorProjectAssetSettings : public UDeveloperSettings
+{
+	GENERATED_BODY()
+
+public:
+	/**
+	 * When performing cleanup operations on redirectors (such as resaving their referencers),
+	 * prompt the user to delete unreferenced redirectors.
+	 */
+	UPROPERTY(EditAnywhere, config, Category = Redirectors)
+	bool bPromptToDeleteUnreferencedRedirectors = true;
+};
 
-UCLASS(config=Editor, meta=(DisplayName="Derived Data"), defaultconfig, MinimalAPI)
+UCLASS(config = Editor, meta = (DisplayName = "Derived Data"), defaultconfig, MinimalAPI)
 class UDDCProjectSettings : public UDeveloperSettings
 {
 	GENERATED_BODY()
