- IAssetTools.h expose bPrompt and create bIgnoreDependencies
- AssetTools.cpp add a case to only copy the selected assets and ignore the dependencies (usecase for automation)
- Added BP example that implements the scripted asset action with the new migrationOption

#tests

diff --git a/Engine/Source/Developer/AssetTools/Private/AssetTools.cpp b/Engine/Source/Developer/AssetTools/Private/AssetTools.cpp
--- a/Engine/Source/Developer/AssetTools/Private/AssetTools.cpp
+++ b/Engine/Source/Developer/AssetTools/Private/AssetTools.cpp
@@ -4873,6 +4873,21 @@ void UAssetToolsImpl::PerformMigratePackages(TArray<FName> PackageNamesToMigrate
 				}
 			}
 
+
+			// Ignore dependencies and only migrate the given assets
+			if (Options.bIgnoreDependencies)
+			{
+				if (PackageNamesToMigrate.Contains(PackageName))
+				{
+					bShouldMigratePackage = true;
+				}
+				else
+				{
+					bShouldMigratePackage = false;
+				}
+			}
+
+
 			FilteredPackageNamesToMove.Add(PackageName);
 
 			if (bShouldMigratePackage)
diff --git a/Engine/Source/Developer/AssetTools/Public/IAssetTools.h b/Engine/Source/Developer/AssetTools/Public/IAssetTools.h
--- a/Engine/Source/Developer/AssetTools/Public/IAssetTools.h
+++ b/Engine/Source/Developer/AssetTools/Public/IAssetTools.h
@@ -219,8 +219,13 @@ struct FMigrationOptions
 	GENERATED_BODY()
 
 	/** Prompt user for confirmation (always false through scripting) */
+	UPROPERTY(BlueprintReadWrite, Category = MigrationOptions)
 	bool bPrompt;
 
+	/** Ignore dependencies of assets, only migrate the given assets. usefull for automation. This will not migrate the actors of a OFPA (one file per actor) level */
+	UPROPERTY(BlueprintReadWrite, Category = MigrationOptions)
+	bool bIgnoreDependencies;
+
 	/** What to do when Assets are conflicting on the destination */
 	UPROPERTY(BlueprintReadWrite, Category = MigrationOptions)
 	EAssetMigrationConflict AssetConflict;
@@ -231,6 +236,7 @@ struct FMigrationOptions
 
 	FMigrationOptions()
 		: bPrompt(false)
+		, bIgnoreDependencies(false)
 		, AssetConflict(EAssetMigrationConflict::Skip)
 	{}
 };
