AssetTools: Fix a bug where FAssetDefinitionProxy::GetIconBrush was calling AssetDefinition->GetThumbnailBrush instead of AssetDefinition->GetIconBrush

diff --git a/Engine/Source/Developer/AssetTools/Private/AssetTools.cpp b/Engine/Source/Developer/AssetTools/Private/AssetTools.cpp
--- a/Engine/Source/Developer/AssetTools/Private/AssetTools.cpp
+++ b/Engine/Source/Developer/AssetTools/Private/AssetTools.cpp
@@ -322,7 +322,7 @@ public:
 	
 	virtual const FSlateBrush* GetIconBrush(const FAssetData& InAssetData, const FName InClassName) const override
 	{
-		return AssetDefinitionPtr->GetThumbnailBrush(InAssetData, InClassName);
+		return AssetDefinitionPtr->GetIconBrush(InAssetData, InClassName);
 	}
 
 	virtual FTopLevelAssetPath GetClassPathName() const override
