UX for updating redirectors
- Refactor AssetViewUtils::LoadAssetsIfNeeded to take settings struct and return enum code.
- Allow cancelling when loading redirectors or when loading referencers.
- Use AssetViewUtils::LoadAssetsIfNeeded in asset view context menu a well as folder tree to unify progress/cancellation for both

diff --git a/Engine/Plugins/Editor/ContentBrowser/ContentBrowserAssetDataSource/Source/ContentBrowserAssetDataSource/Private/AssetFileContextMenu.cpp b/Engine/Plugins/Editor/ContentBrowser/ContentBrowserAssetDataSource/Source/ContentBrowserAssetDataSource/Private/AssetFileContextMenu.cpp
--- a/Engine/Plugins/Editor/ContentBrowser/ContentBrowserAssetDataSource/Source/ContentBrowserAssetDataSource/Private/AssetFileContextMenu.cpp
+++ b/Engine/Plugins/Editor/ContentBrowser/ContentBrowserAssetDataSource/Source/ContentBrowserAssetDataSource/Private/AssetFileContextMenu.cpp
@@ -1338,19 +1338,19 @@ void FAssetFileContextMenu::ExecuteCreateBlueprintUsing()
 
 void FAssetFileContextMenu::GetSelectedAssets(TArray<UObject*>& Assets, bool SkipRedirectors) const
 {
-	TArray<FString> SelectedAssetPaths;
-	for (const FAssetData& SelectedAsset : SelectedAssets)
+	AssetViewUtils::FLoadAssetsSettings Settings{
+		// Default settings
+	};
+	if (SkipRedirectors)
 	{
-		if (SkipRedirectors && (SelectedAsset.AssetClassPath == UObjectRedirector::StaticClass()->GetClassPathName()))
-		{
-			// Don't operate on Redirectors
-			continue;
-		}
-
-		SelectedAssetPaths.Add(SelectedAsset.GetObjectPathString());
+		TArray<FAssetData> FilteredAssets;
+		Algo::CopyIf(SelectedAssets, FilteredAssets, [](const FAssetData& Asset) { return !Asset.IsRedirector(); });
+		AssetViewUtils::LoadAssetsIfNeeded(FilteredAssets, Assets, Settings);
+	}
+	else
+	{
+		AssetViewUtils::LoadAssetsIfNeeded(SelectedAssets, Assets, Settings);
 	}
-
-	AssetViewUtils::LoadAssetsIfNeeded(SelectedAssetPaths, Assets);
 }
 
 void FAssetFileContextMenu::GetSelectedAssetData(TArray<FAssetData>& AssetDataList, bool SkipRedirectors) const
diff --git a/Engine/Plugins/Editor/ContentBrowser/ContentBrowserAssetDataSource/Source/ContentBrowserAssetDataSource/Private/AssetFolderContextMenu.cpp b/Engine/Plugins/Editor/ContentBrowser/ContentBrowserAssetDataSource/Source/ContentBrowserAssetDataSource/Private/AssetFolderContextMenu.cpp
--- a/Engine/Plugins/Editor/ContentBrowser/ContentBrowserAssetDataSource/Source/ContentBrowserAssetDataSource/Private/AssetFolderContextMenu.cpp
+++ b/Engine/Plugins/Editor/ContentBrowser/ContentBrowserAssetDataSource/Source/ContentBrowserAssetDataSource/Private/AssetFolderContextMenu.cpp
@@ -146,27 +146,26 @@ void FAssetFolderContextMenu::ExecuteFixUpRedirectorsInFolder()
 	{
 		return;
 	}
-	
-	TArray<FString> ObjectPaths;
-	for (const FAssetData& Asset : AssetList)
-	{
-		ObjectPaths.Add(Asset.GetObjectPathString());
-	}
 
 	FScopedSlowTask SlowTask(3, LOCTEXT("FixupRedirectorsSlowTask", "Fixing up redirectors"));
-	SlowTask.MakeDialog();
+	SlowTask.MakeDialog(true);
 
-	TArray<UObject*> Objects;
-	const bool bAllowedToPromptToLoadAssets = true;
-	const bool bLoadRedirects = true;
 	SlowTask.EnterProgressFrame(1, LOCTEXT("FixupRedirectors_LoadAssets", "Loading Assets..."));
-	if (AssetViewUtils::LoadAssetsIfNeeded(ObjectPaths, Objects, bAllowedToPromptToLoadAssets, bLoadRedirects))
+	TArray<UObject*> Objects;
+	AssetViewUtils::FLoadAssetsSettings Settings{
+		.bFollowRedirectors = false,
+		.bAllowCancel = true,
+	};
+	AssetViewUtils::ELoadAssetsResult Result = AssetViewUtils::LoadAssetsIfNeeded(AssetList, Objects, Settings);
+	if (Result != AssetViewUtils::ELoadAssetsResult::Cancelled && !SlowTask.ShouldCancel())
 	{
-		// Transform Objects array to ObjectRedirectors array
 		TArray<UObjectRedirector*> Redirectors;
 		for (UObject* Object : Objects)
 		{
-			Redirectors.Add(CastChecked<UObjectRedirector>(Object));
+			if (UObjectRedirector* Redirector = Cast<UObjectRedirector>(Object))
+			{
+				Redirectors.Add(Redirector);
+			}
 		}
 
 		SlowTask.EnterProgressFrame(1, LOCTEXT("FixupRedirectors_FixupReferencers", "Fixing up referencers..."));
diff --git a/Engine/Plugins/Editor/EngineAssetDefinitions/Source/Private/AssetDefinition_Redirector.cpp b/Engine/Plugins/Editor/EngineAssetDefinitions/Source/Private/AssetDefinition_Redirector.cpp
--- a/Engine/Plugins/Editor/EngineAssetDefinitions/Source/Private/AssetDefinition_Redirector.cpp
+++ b/Engine/Plugins/Editor/EngineAssetDefinitions/Source/Private/AssetDefinition_Redirector.cpp
@@ -3,6 +3,7 @@
 #include "AssetDefinition_Redirector.h"
 
 #include "AssetToolsModule.h"
+#include "AssetViewUtils.h"
 #include "ContentBrowserMenuContexts.h"
 #include "IAssetTools.h"
 #include "ToolMenus.h"
@@ -54,12 +55,33 @@ namespace MenuExtension_Redirector
 
 		if (const UContentBrowserAssetContextMenuContext* Context = UContentBrowserAssetContextMenuContext::FindContextWithAssets(MenuContext))
 		{
-			TArray<UObjectRedirector*> Redirectors = Context->LoadSelectedObjects<UObjectRedirector>();
+			FScopedSlowTask SlowTask(3, LOCTEXT("FixupRedirectorsSlowTask", "Fixing up redirectors"));
+			SlowTask.MakeDialog(true);
 
-			if (Redirectors.Num() > 0)
+			SlowTask.EnterProgressFrame(1, LOCTEXT("FixupRedirectors_LoadAssets", "Loading Assets..."));
+			TArray<UObject*> Objects;
+			AssetViewUtils::FLoadAssetsSettings Settings{
+				.bFollowRedirectors = false,
+				.bAllowCancel = true,
+			};
+			AssetViewUtils::ELoadAssetsResult Result = AssetViewUtils::LoadAssetsIfNeeded(Context->SelectedAssets, Objects, Settings);
+			if (Result != AssetViewUtils::ELoadAssetsResult::Cancelled && !SlowTask.ShouldCancel())
 			{
-				IAssetTools& AssetTools = FModuleManager::LoadModuleChecked<FAssetToolsModule>("AssetTools").Get();
-				AssetTools.FixupReferencers(Redirectors, /*bCheckoutDialogPrompt=*/true, ERedirectFixupMode::PromptForDeletingRedirectors);
+				TArray<UObjectRedirector*> Redirectors;
+				for (UObject* Object : Objects)
+				{
+					if (UObjectRedirector* Redirector = Cast<UObjectRedirector>(Object))
+					{
+						Redirectors.Add(Redirector);
+					}
+				}
+
+				if (Redirectors.Num() > 0)
+				{
+					SlowTask.EnterProgressFrame(1, LOCTEXT("FixupRedirectors_FixupReferencers", "Fixing up referencers..."));
+					IAssetTools& AssetTools = FModuleManager::LoadModuleChecked<FAssetToolsModule>("AssetTools").Get();
+					AssetTools.FixupReferencers(Redirectors, /*bCheckoutDialogPrompt=*/true, ERedirectFixupMode::PromptForDeletingRedirectors);
+				}
 			}
 		}
 	}
diff --git a/Engine/Source/Developer/AssetTools/Private/AssetFixUpRedirectors.cpp b/Engine/Source/Developer/AssetTools/Private/AssetFixUpRedirectors.cpp
--- a/Engine/Source/Developer/AssetTools/Private/AssetFixUpRedirectors.cpp
+++ b/Engine/Source/Developer/AssetTools/Private/AssetFixUpRedirectors.cpp
@@ -623,15 +623,21 @@ void FAssetFixUpRedirectors::ExecuteFixUp(TArray<TWeakObjectPtr<UObjectRedirecto
 	// Load all referencing packages.
 	TSet<UPackage*> ReferencingPackagesToSave;
 	TSet<UPackage*> LoadedPackages;
+	bool bCancel = false;
 	{
 		FScopedSlowTask SlowTask(static_cast<float>(RedirectorRefsList.Num()), LOCTEXT( "LoadingReferencingPackages", "Loading Referencing Packages..." ) );
-		SlowTask.MakeDialog();
+		SlowTask.MakeDialog(true);
 		ISourceControlProvider& SourceControlProvider = ISourceControlModule::Get().GetProvider();
 
 		// Load all packages that reference each redirector, if possible
 		for (FRedirectorRefs& RedirectorRefs : RedirectorRefsList)
 		{
 			SlowTask.EnterProgressFrame(1);
+			if (SlowTask.ShouldCancel())
+			{
+				bCancel = true;
+				break;
+			}
 			if (bMayDeleteRedirectors && ISourceControlModule::Get().IsEnabled())
 			{
 				FSourceControlStatePtr SourceControlState = SourceControlProvider.GetState(RedirectorRefs.Redirector->GetOutermost(), EStateCacheUsage::Use);
@@ -640,6 +646,7 @@ void FAssetFixUpRedirectors::ExecuteFixUp(TArray<TWeakObjectPtr<UObjectRedirecto
 				if (!bValidSCCState)
 				{
 					RedirectorRefs.bSCCError = true;
+					// Continue to load the referencers because we may still be able to fix them up 
 				}
 			}
 
@@ -692,6 +699,11 @@ void FAssetFixUpRedirectors::ExecuteFixUp(TArray<TWeakObjectPtr<UObjectRedirecto
 			}
 		}
 	};
+	
+	if (bCancel)
+	{
+		return;
+	}
 
 	// Add all referencing packages objects that aren't RF_Standalone to the root set to avoid them being GC'd during the following processing
 	TArray<TStrongObjectPtr<UObject>> RootedObjects;
diff --git a/Engine/Source/Developer/AssetTools/Private/AssetViewUtils.cpp b/Engine/Source/Developer/AssetTools/Private/AssetViewUtils.cpp
--- a/Engine/Source/Developer/AssetTools/Private/AssetViewUtils.cpp
+++ b/Engine/Source/Developer/AssetTools/Private/AssetViewUtils.cpp
@@ -1,6 +1,7 @@
 // Copyright Epic Games, Inc. All Rights Reserved.
 
 #include "AssetViewUtils.h"
+#include "Algo/Transform.h"
 #include "HAL/IConsoleManager.h"
 #include "Misc/MessageDialog.h"
 #include "HAL/FileManager.h"
@@ -43,7 +44,7 @@
 #include "Editor.h"
 #include "UObject/LinkerInstancingContext.h"
 
-#define LOCTEXT_NAMESPACE "ContentBrowser"
+#define LOCTEXT_NAMESPACE "AssetViewUtils"
 
 DEFINE_LOG_CATEGORY_STATIC(LogAssetViewTools, Warning, Warning);
 
@@ -72,28 +73,26 @@ namespace AssetViewUtils
 
 	/** Get all the objects in a list of asset data with optional load of all external packages */
 	void GetObjectsInAssetData(const TArray<FAssetData>& AssetList, TArray<UObject*>& OutDroppedObjects, bool bLoadAllExternalObjects);
-
-	/** Makes sure the specified assets are loaded into memory. */
-	bool LoadAssetsIfNeeded(const TArray<FString>& ObjectPaths, TArray<UObject*>& LoadedObjects, bool bAllowedToPromptToLoadAssets, bool bLoadRedirects, bool bLoadWorldPartitionWorlds, bool bLoadAllExternalObjects);
 }
 
 bool AssetViewUtils::OpenEditorForAsset(const FString& ObjectPath)
 {
 	// Load the asset if unloaded
 	TArray<UObject*> LoadedObjects;
-	TArray<FString> ObjectPaths;
-	ObjectPaths.Add(ObjectPath);
+	TArray<FSoftObjectPath> ObjectPaths;
+	ObjectPaths.Emplace(ObjectPath);
 
 	// Here we want to load the asset as it will be passed to OpenEditorForAsset
-	const bool bAllowedToPromptToLoadAssets = true;
-	const bool bLoadRedirects = false;
-	const bool bLoadWorldPartitionWorlds = true;
-	const bool bLoadAllExternalObjects = false;
-	LoadAssetsIfNeeded(ObjectPaths, LoadedObjects, bAllowedToPromptToLoadAssets, bLoadRedirects, bLoadWorldPartitionWorlds, bLoadAllExternalObjects);
+	FLoadAssetsSettings Settings{
+		.bFollowRedirectors = false,
+		.bLoadWorldPartitionMaps = true,
+		.bLoadAllExternalObjects = false,
+	};
+	LoadAssetsIfNeeded(ObjectPaths, LoadedObjects, Settings);
 
 	// Open the editor for the specified asset
-	UObject* FoundObject = FindObject<UObject>(NULL, *ObjectPath);
-			
+	UObject* FoundObject = FindObject<UObject>(nullptr, *ObjectPath);
+
 	return OpenEditorForAsset(FoundObject);
 }
 
@@ -124,105 +123,178 @@ bool AssetViewUtils::OpenEditorForAsset(const TArray<UObject*>& Assets)
 
 bool AssetViewUtils::LoadAssetsIfNeeded(const TArray<FString>& ObjectPaths, TArray<UObject*>& LoadedObjects, bool bAllowedToPromptToLoadAssets, bool bLoadRedirects)
 {
-	const bool bLoadWorldPartitionWorlds = false;
-	const bool bLoadAllExternalObjects = false;
-	return LoadAssetsIfNeeded(ObjectPaths, LoadedObjects, bAllowedToPromptToLoadAssets, bLoadRedirects, bLoadWorldPartitionWorlds, bLoadAllExternalObjects);
+	FLoadAssetsSettings Settings{
+		.bFollowRedirectors = bLoadRedirects,
+		.bLoadWorldPartitionMaps = false,
+		.bLoadAllExternalObjects = false,
+	};
+
+	switch (LoadAssetsIfNeeded(ObjectPaths, LoadedObjects, Settings))
+	{
+		case ELoadAssetsResult::Success:
+			return true;
+		case ELoadAssetsResult::Cancelled:  // fallthrough
+		case ELoadAssetsResult::SomeFailed: // fallthrough
+			return false;
+		default:
+			check("Unhandled return value from LoadAssetsIfNeeded");
+			return false;
+	}
 }
 
-bool AssetViewUtils::LoadAssetsIfNeeded(const TArray<FString>& ObjectPaths, TArray<UObject*>& LoadedObjects, bool bAllowedToPromptToLoadAssets, bool bLoadRedirects, bool bLoadWorldPartitionWorlds, bool bLoadAllExternalObjects)
+AssetViewUtils::ELoadAssetsResult AssetViewUtils::LoadAssetsIfNeeded(TConstArrayView<FString> ObjectPathStrings, TArray<UObject*>& OutLoadedObjects, const FLoadAssetsSettings& Settings)
 {
-	bool bAnyObjectsWereLoadedOrUpdated = false;
+	TArray<FSoftObjectPath> Paths;
+	Paths.Reserve(ObjectPathStrings.Num());
+	Algo::Transform(ObjectPathStrings, Paths, UE_PROJECTION(FSoftObjectPath::FSoftObjectPath));
+	return LoadAssetsIfNeeded(Paths, OutLoadedObjects, Settings);
+}
 
-	// Build a list of unloaded assets
-	TArray<FString> UnloadedObjectPaths;
-	bool bAtLeastOneUnloadedMap = false;
-	for (int32 PathIdx = 0; PathIdx < ObjectPaths.Num(); ++PathIdx)
+AssetViewUtils::ELoadAssetsResult AssetViewUtils::LoadAssetsIfNeeded(TConstArrayView<FSoftObjectPath> ObjectPaths, TArray<UObject*>& OutLoadedObjects, const FLoadAssetsSettings& Settings)
+{
+	IAssetRegistry& AssetRegistry = FModuleManager::LoadModuleChecked<FAssetRegistryModule>(TEXT("AssetRegistry")).Get();
+	TArray<FAssetData> Assets;
+	Assets.Reserve(ObjectPaths.Num());
+	TArray<FSoftObjectPath> ToScan;
+	for (const FSoftObjectPath& Path : ObjectPaths)
 	{
-		const FString& ObjectPath = ObjectPaths[PathIdx];
-
-		UObject* FoundObject = FindObject<UObject>(NULL, *ObjectPath);
-		if ( FoundObject )
+		FAssetData Asset = AssetRegistry.GetAssetByObjectPath(Path, true);
+		if (Asset.IsValid())
 		{
-			LoadedObjects.Add(FoundObject);
+			Assets.Emplace(MoveTemp(Asset));
 		}
 		else
 		{
-			if ( FEditorFileUtils::IsMapPackageAsset(ObjectPath) )
+			ToScan.Emplace(Path);
+		}
+	}
+	if (ToScan.Num() > 0)
+	{
+		TArray<FString> ScanPaths;
+		Algo::Transform(ToScan, ScanPaths, UE_PROJECTION_MEMBER(FSoftObjectPath, GetLongPackageName));
+		AssetRegistry.ScanFilesSynchronous(ScanPaths);
+		for (const FSoftObjectPath& Path : ToScan)
+		{
+			FAssetData Asset = AssetRegistry.GetAssetByObjectPath(Path, true);
+			if (Asset.IsValid())
 			{
-				FName PackageName = FName(*FEditorFileUtils::ExtractPackageName(ObjectPath));
-				if (!bLoadWorldPartitionWorlds && ULevel::GetIsLevelPartitionedFromPackage(PackageName))
-				{
-					continue;
-				}
-				
-				bAtLeastOneUnloadedMap = true;
+				Assets.Emplace(MoveTemp(Asset));
 			}
-
-			// Unloaded asset, we will load it later
-			UnloadedObjectPaths.Add(ObjectPath);
 		}
 	}
+	return LoadAssetsIfNeeded(Assets, OutLoadedObjects, Settings);
+}
+
+AssetViewUtils::ELoadAssetsResult AssetViewUtils::LoadAssetsIfNeeded(TConstArrayView<FAssetData> Assets, TArray<UObject*>& OutLoadedObjects, const FLoadAssetsSettings& Settings)
+{
+	bool bAnySucceeded = false;
 
-	// Make sure all selected objects are loaded, where possible
-	if ( UnloadedObjectPaths.Num() > 0 )
+	// Build a list of unloaded assets
+	TArray<FSoftObjectPath> UnloadedObjectPaths;
+	bool bAtLeastOneUnloadedMap = false;
+	for (const FAssetData& Asset : Assets)
 	{
-		// Get the maximum objects to load before displaying the slow task
-		const bool bShowProgressDialog = (UnloadedObjectPaths.Num() > GetDefault<UContentBrowserSettings>()->NumObjectsToLoadBeforeWarning) || bAtLeastOneUnloadedMap;
-		FScopedSlowTask SlowTask(static_cast<float>(UnloadedObjectPaths.Num()), LOCTEXT("LoadingObjects", "Loading Objects..."));
-		if (bShowProgressDialog)
+		UObject* FoundObject = Asset.FastGetAsset(false);
+		if (FoundObject)
+		{
+			OutLoadedObjects.Add(FoundObject);
+			bAnySucceeded = true;
+			continue;
+		}
+		else if (!Settings.bLoadWorldPartitionMaps && ULevel::GetIsLevelPartitionedFromAsset(Asset))
+		{
+			// Skip
+			bAtLeastOneUnloadedMap = true;
+		}
+		else
 		{
-			SlowTask.MakeDialog();
+			UnloadedObjectPaths.Add(Asset.GetSoftObjectPath());
 		}
+	}
+
+	if (UnloadedObjectPaths.Num() == 0)
+	{
+		return ELoadAssetsResult::Success;
+	}
 
-		bool bSomeObjectsFailedToLoad = false;
+	if (Settings.bAlwaysPromptBeforeLoading || UnloadedObjectPaths.Num() > GetDefault<UContentBrowserSettings>()->NumObjectsToLoadBeforeWarning)
+	{
+		EAppReturnType::Type Decision = FMessageDialog::Open(EAppMsgType::YesNo,
+			FText::Format(LOCTEXT("LoadingManyAssets", "About to load {0} assets. Are you sure you want to continue?"), UnloadedObjectPaths.Num()),
+			LOCTEXT("LoadingManyAssetsTitle", "Loading Many Assets"));
+
+		switch (Decision)
 		{
-			TGuardValue<bool> IsEditorLoadingPackageGuard(GIsEditorLoadingPackage, true);
+			case EAppReturnType::Cancel: // fallthrough
+			case EAppReturnType::No:     // fallthrough
+			case EAppReturnType::NoAll:
+				return ELoadAssetsResult::Cancelled;
+			case EAppReturnType::Yes:      // fallthrough
+			case EAppReturnType::YesAll:   // fallthrough
+			case EAppReturnType::Ok:       // fallthrough
+			case EAppReturnType::Retry:    // fallthrough
+			case EAppReturnType::Continue: // fallthrough
+				break;
+		}
+	}
 
-			// We usually don't want to follow redirects when loading objects for the Content Browser.  It would
-			// allow a user to interact with a ghost/unverified asset as if it were still alive.
-			// This can be overridden by providing bLoadRedirects = true as a parameter.
-			const ELoadFlags LoadFlags = bLoadRedirects ? LOAD_None : LOAD_NoRedirects;
+	FScopedSlowTask SlowTask(static_cast<float>(UnloadedObjectPaths.Num()), LOCTEXT("LoadingAssets", "Loading Assets..."));
+	// Always make dialog, even a single asset can be slow to load
+	SlowTask.MakeDialog(Settings.bAllowCancel);
 
-			for (int32 PathIdx = 0; PathIdx < UnloadedObjectPaths.Num(); ++PathIdx)
-			{
-				const FString& ObjectPath = UnloadedObjectPaths[PathIdx];
-				SlowTask.EnterProgressFrame(1, FText::Format(LOCTEXT("LoadingObjectf", "Loading {0}..."), FText::FromString(ObjectPath)));
+	bool bSomeObjectsFailedToLoad = false;
+	bool bCancelled = false;
+	{
+		TGuardValue<bool> IsEditorLoadingPackageGuard(GIsEditorLoadingPackage, true);
+		const ELoadFlags LoadFlags = Settings.bFollowRedirectors ? LOAD_None : LOAD_NoRedirects;
 
-				// Load up the object
-				FLinkerInstancingContext InstancingContext(bLoadAllExternalObjects ? TSet<FName>{ ULevel::LoadAllExternalObjectsTag } : TSet<FName>());
-				UObject* LoadedObject = LoadObject<UObject>(NULL, *ObjectPath, NULL, LoadFlags, NULL, &InstancingContext);
-				if ( LoadedObject )
-				{
-					LoadedObjects.Add(LoadedObject);
-				}
-				else
-				{
-					bSomeObjectsFailedToLoad = true;
-				}
+		for (const FSoftObjectPath& ObjectPath : UnloadedObjectPaths)
+		{
+			SlowTask.EnterProgressFrame(1, FText::Format(LOCTEXT("LoadingAsset", "Loading {0}..."), FText::FromName(ObjectPath.GetLongPackageFName())));
 
-				if (GWarn->ReceivedUserCancel())
-				{
-					// If the user has canceled stop loading the remaining objects. We don't add the remaining objects to the failed string,
-					// this would only result in launching another dialog when by their actions the user clearly knows not all of the 
-					// assets will have been loaded.
-					break;
-				}
+			// Load up the object
+			FLinkerInstancingContext InstancingContext(Settings.bLoadAllExternalObjects ? TSet<FName>{ ULevel::LoadAllExternalObjectsTag } : TSet<FName>());
+			UObject* LoadedObject = LoadObject<UObject>(nullptr, *ObjectPath.ToString(), nullptr, LoadFlags, nullptr, &InstancingContext);
+			if (LoadedObject)
+			{
+				OutLoadedObjects.Add(LoadedObject);
+				bAnySucceeded = true;
+			}
+			else
+			{
+				bSomeObjectsFailedToLoad = true;
+			}
+
+			if (SlowTask.ShouldCancel())
+			{
+				bCancelled = true;
+				break;
 			}
 		}
+	}
 
-		if ( bSomeObjectsFailedToLoad )
-		{
-			FNotificationInfo Info(LOCTEXT("LoadObjectFailed", "Failed to load assets"));
-			Info.ExpireDuration = 5.0f;
-			Info.Hyperlink = FSimpleDelegate::CreateStatic([](){ FMessageLog("LoadErrors").Open(EMessageSeverity::Info, true); });
-			Info.HyperlinkText = LOCTEXT("LoadObjectHyperlink", "Show Message Log");
+	if (bSomeObjectsFailedToLoad)
+	{
+		FNotificationInfo Info(LOCTEXT("SomeLoadsFailed", "Failed to load some assets"));
+		Info.ExpireDuration = 5.0f;
+		Info.Hyperlink = FSimpleDelegate::CreateStatic([]() { FMessageLog("LoadErrors").Open(EMessageSeverity::Info, true); });
+		Info.HyperlinkText = LOCTEXT("LoadObjectHyperlink", "Show Message Log");
 
-			FSlateNotificationManager::Get().AddNotification(Info);
-			return false;
-		}
+		FSlateNotificationManager::Get().AddNotification(Info);
 	}
 
-	return true;
+	if (bCancelled)
+	{
+		return ELoadAssetsResult::Cancelled;
+	}
+	else if (bSomeObjectsFailedToLoad)
+	{
+		return ELoadAssetsResult::SomeFailed;
+	}
+	else
+	{
+		return ELoadAssetsResult::Success;
+	}
 }
 
 void AssetViewUtils::GetUnloadedAssets(const TArray<FString>& ObjectPaths, TArray<FString>& OutUnloadedObjects)
@@ -380,11 +452,12 @@ bool AssetViewUtils::DeleteFolders(const TArray<FString>& PathsToDelete)
 
 		// Load all the assets in the selected paths
 		TArray<UObject*> LoadedAssets;
-		const bool bAllowedToPromptToLoadAssets = true;
-		const bool bLoadRedirects = false;
-		const bool bLoadWorldPartitionWorlds = true;
-		const bool bLoadAllExternalObjects = true;
-		if ( LoadAssetsIfNeeded(ObjectPaths, LoadedAssets, bAllowedToPromptToLoadAssets, bLoadRedirects, bLoadWorldPartitionWorlds, bLoadAllExternalObjects))
+		FLoadAssetsSettings Settings{
+			.bFollowRedirectors = false,
+			.bLoadWorldPartitionMaps = true,
+			.bLoadAllExternalObjects = true,
+		};
+		if (LoadAssetsIfNeeded(ObjectPaths, LoadedAssets, Settings) == ELoadAssetsResult::Success)
 		{
 			// Make sure we loaded all of them
 			if ( LoadedAssets.Num() == NumAssetsInPaths )
@@ -781,11 +854,12 @@ bool AssetViewUtils::PrepareFoldersForDragDrop(const TArray<FString>& SourcePath
 
 		// Load all assets in this path if needed
 		TArray<UObject*> AllLoadedAssets;
-		const bool bAllowedToPromptToLoadAssets = false;
-		const bool bLoadRedirects = false;
-		const bool bLoadWorldPartitionWorlds = true;
-		const bool bLoadAllExternalObjects = true;
-		LoadAssetsIfNeeded(ObjectPaths, AllLoadedAssets, bAllowedToPromptToLoadAssets, bLoadRedirects, bLoadWorldPartitionWorlds, bLoadAllExternalObjects);
+		FLoadAssetsSettings Settings{
+			.bFollowRedirectors = false,
+			.bLoadWorldPartitionMaps = true,
+			.bLoadAllExternalObjects = true
+		};
+		LoadAssetsIfNeeded(ObjectPaths, AllLoadedAssets, Settings);
 
 		// Add a slash to the end of the path so StartsWith doesn't get a false positive on similarly named folders
 		const FString SourcePathWithSlash = *PathIt + TEXT("/");
diff --git a/Engine/Source/Developer/AssetTools/Public/AssetViewUtils.h b/Engine/Source/Developer/AssetTools/Public/AssetViewUtils.h
--- a/Engine/Source/Developer/AssetTools/Public/AssetViewUtils.h
+++ b/Engine/Source/Developer/AssetTools/Public/AssetViewUtils.h
@@ -19,6 +19,43 @@ namespace AssetViewUtils
 	/** Opens the asset editor for the specified assets */
 	ASSETTOOLS_API bool OpenEditorForAsset(const TArray<UObject*>& Assets);
 
+	struct FLoadAssetsSettings
+	{
+		// Prompt the user with the number of unloaded assets before continuing to load
+		bool bAlwaysPromptBeforeLoading = false;
+		// Show a notification to the user that some assets failed to load
+		bool bShowFailureNotification = true;
+		// Follow redirectors and return the target assets.
+		// We usually don't want to follow redirects when loading objects for the Content Browser.  It would
+		// allow a user to interact with a ghost/unverified asset as if it were still alive.
+		bool bFollowRedirectors = false;
+		// Allow the user to cancel the load operation
+		bool bAllowCancel = false;
+		// Advanced setting: allow loading of world partition map packages.
+		bool bLoadWorldPartitionMaps = false;
+		// Advanced setting: load all external objects/actors with e.g. world partition map packages.
+		bool bLoadAllExternalObjects = false;
+	};
+
+	enum class ELoadAssetsResult : uint8
+	{
+		Success,    // The user didn't cancel and all requested assets were successfully loaded
+		Cancelled,  // The user cancelled the load operation
+		SomeFailed, // Some or all of the assets failed to load
+	};
+
+	/**
+	 * Makes sure the specified assets are loaded into memory, with displayed progress and optional cancellation for the user.
+	 * Prefer the overload taking FAssetData if possible.
+	 *
+	 * @param Assets The assets to load
+	 * @param LoadedObjects The returned list of objects that were already loaded or loaded by this method.
+	 * @return enum specifying whether the load was successful, cancelled or if some or all assets failed to load
+	 */
+	ASSETTOOLS_API ELoadAssetsResult LoadAssetsIfNeeded(TConstArrayView<FAssetData> Assets, TArray<UObject*>& LoadedObjects, const FLoadAssetsSettings& InSettings);
+	ASSETTOOLS_API ELoadAssetsResult LoadAssetsIfNeeded(TConstArrayView<FString> ObjectPaths, TArray<UObject*>& LoadedObjects, const FLoadAssetsSettings& InSettings);
+	ASSETTOOLS_API ELoadAssetsResult LoadAssetsIfNeeded(TConstArrayView<FSoftObjectPath> ObjectPaths, TArray<UObject*>& LoadedObjects, const FLoadAssetsSettings& InSettings);
+
 	/**
 	  * Makes sure the specified assets are loaded into memory.
 	  * 
@@ -26,6 +63,7 @@ namespace AssetViewUtils
 	  * @param LoadedObjects The returned list of objects that were already loaded or loaded by this method.
 	  * @return false if user canceled after being warned about loading very many packages.
 	  */
+	UE_DEPRECATED(5.5, "Use LoadAssetsIfNeeded taking struct FLoadAssetsSettings instead.")
 	ASSETTOOLS_API bool LoadAssetsIfNeeded(const TArray<FString>& ObjectPaths, TArray<UObject*>& LoadedObjects, bool bAllowedToPromptToLoadAssets = true, bool bLoadRedirects = false);
 
 	/**
diff --git a/Engine/Source/Editor/AnimationModifiers/Private/AnimationModifiersModule.cpp b/Engine/Source/Editor/AnimationModifiers/Private/AnimationModifiersModule.cpp
--- a/Engine/Source/Editor/AnimationModifiers/Private/AnimationModifiersModule.cpp
+++ b/Engine/Source/Editor/AnimationModifiers/Private/AnimationModifiersModule.cpp
@@ -2,6 +2,7 @@
 
 #include "AnimationModifiersModule.h"
 
+#include "Algo/Copy.h"
 #include "Animation/AnimSequence.h"
 #include "AnimationModifier.h"
 #include "AnimationModifierDetailCustomization.h"
@@ -166,20 +167,16 @@ void FAnimationModifiersModule::RegisterMenus()
 				return;
 			}
 
-			TArray<FString> AnimSequencePaths;
-			Algo::TransformIf(Context->SelectedAssets, AnimSequencePaths, [](const FAssetData& AssetData)
+			TArray<FAssetData> AnimSequenceAssets;
+			Algo::CopyIf(Context->SelectedAssets, AnimSequenceAssets, [](const FAssetData& AssetData)
 			{
 				return AssetData.AssetClassPath == UAnimSequence::StaticClass()->GetClassPathName();
-			},
-			[](const FAssetData& AssetData)
-			{
-				return AssetData.GetObjectPathString();
 			});
 			
-			auto GetAnimSequences = [AnimSequencePaths](TArray<UAnimSequence*>& OutSequences)
+			auto GetAnimSequences = [AnimSequenceAssets](TArray<UAnimSequence*>& OutSequences)
 			{
 				TArray<UObject*> Objects;
-				AssetViewUtils::LoadAssetsIfNeeded(AnimSequencePaths, Objects);
+				AssetViewUtils::LoadAssetsIfNeeded(AnimSequenceAssets, Objects, AssetViewUtils::FLoadAssetsSettings{});
 			
 				Algo::TransformIf(Objects, OutSequences, 
 				[](const UObject* Object)
diff --git a/Engine/Source/Editor/ContentBrowser/Private/ContentBrowserMenuContexts.cpp b/Engine/Source/Editor/ContentBrowser/Private/ContentBrowserMenuContexts.cpp
--- a/Engine/Source/Editor/ContentBrowser/Private/ContentBrowserMenuContexts.cpp
+++ b/Engine/Source/Editor/ContentBrowser/Private/ContentBrowserMenuContexts.cpp
@@ -33,14 +33,11 @@ bool UContentBrowserToolbarMenuContext::CanWriteToCurrentPath() const
 
 TArray<UObject*> UContentBrowserAssetContextMenuContext::LoadSelectedObjectsIfNeeded() const
 {
-	TArray<FString> SelectedAssetPaths;
-	for (const FAssetData& SelectedAsset : SelectedAssets)
-	{
-		SelectedAssetPaths.Add(SelectedAsset.GetObjectPathString());
-	}
-
 	TArray<UObject*> Objects;
-	AssetViewUtils::LoadAssetsIfNeeded(SelectedAssetPaths, Objects);
+	AssetViewUtils::FLoadAssetsSettings Settings{
+		// Default settings
+	};
+	AssetViewUtils::LoadAssetsIfNeeded(SelectedAssets, Objects, Settings);
 
 	return Objects;
 }
