Fix UE_PROJECTION with constructor.

diff --git a/Engine/Source/Developer/AssetTools/Private/AssetViewUtils.cpp b/Engine/Source/Developer/AssetTools/Private/AssetViewUtils.cpp
--- a/Engine/Source/Developer/AssetTools/Private/AssetViewUtils.cpp
+++ b/Engine/Source/Developer/AssetTools/Private/AssetViewUtils.cpp
@@ -146,7 +146,7 @@ AssetViewUtils::ELoadAssetsResult AssetViewUtils::LoadAssetsIfNeeded(TConstArray
 {
 	TArray<FSoftObjectPath> Paths;
 	Paths.Reserve(ObjectPathStrings.Num());
-	Algo::Transform(ObjectPathStrings, Paths, UE_PROJECTION(FSoftObjectPath::FSoftObjectPath));
+	Algo::Transform(ObjectPathStrings, Paths, UE_PROJECTION(FSoftObjectPath));
 	return LoadAssetsIfNeeded(Paths, OutLoadedObjects, Settings);
 }
 
