Workaround for crash serializing CDOs: Skip text properties

diff --git a/Engine/Source/Developer/AssetTools/Private/AssetRenameManager.cpp b/Engine/Source/Developer/AssetTools/Private/AssetRenameManager.cpp
--- a/Engine/Source/Developer/AssetTools/Private/AssetRenameManager.cpp
+++ b/Engine/Source/Developer/AssetTools/Private/AssetRenameManager.cpp
@@ -537,7 +537,7 @@ struct FSoftObjectPathRenameSerializer : public FArchiveUObject
 		}
 
 		FFieldClass* PropertyClass = InProperty->GetClass();
-		if (PropertyClass->GetCastFlags() & (CASTCLASS_FBoolProperty | CASTCLASS_FNameProperty | CASTCLASS_FStrProperty | CASTCLASS_FMulticastDelegateProperty))
+		if (PropertyClass->GetCastFlags() & (CASTCLASS_FBoolProperty | CASTCLASS_FNameProperty | CASTCLASS_FStrProperty | CASTCLASS_FTextProperty | CASTCLASS_FMulticastDelegateProperty))
 		{
 			return true;
 		}
