Fix use of deprecated method

diff --git a/Engine/Source/Developer/AssetTools/Private/AssetFixUpRedirectors.cpp b/Engine/Source/Developer/AssetTools/Private/AssetFixUpRedirectors.cpp
--- a/Engine/Source/Developer/AssetTools/Private/AssetFixUpRedirectors.cpp
+++ b/Engine/Source/Developer/AssetTools/Private/AssetFixUpRedirectors.cpp
@@ -852,9 +852,10 @@ void FAssetFixUpRedirectors::ExecuteFixUp(TArray<TWeakObjectPtr<UObjectRedirecto
 			for (const UObjectRedirector* Redirector = RedirectorRefs.Redirector.Get(); Redirector; Redirector = Cast<UObjectRedirector>(Redirector->DestinationObject))
 			{
 				const FSoftObjectPath RedirectorObjectPath = FSoftObjectPath(Redirector);
-				if (!CollectionManagerModule.Get().HandleRedirectorDeleted(RedirectorObjectPath))
+				FText Error;
+				if (!CollectionManagerModule.Get().HandleRedirectorDeleted(RedirectorObjectPath, &Error))
 				{
-					RedirectorRefs.OtherFailures.Add(FText::Format(LOCTEXT("RedirectorFixupFailed_CollectionsFailedToSave", "Referencing collection(s) failed to save: {0}"), CollectionManagerModule.Get().GetLastError()));
+					RedirectorRefs.OtherFailures.Add(FText::Format(LOCTEXT("RedirectorFixupFailed_CollectionsFailedToSave", "Referencing collection(s) failed to save: {0}"), Error));
 				}
 			}
 		}
