[Backout] - CL33009366
Original CL Desc
-----------------------------------------------------------------
Fix use of deprecated method

diff --git a/Engine/Source/Developer/AssetTools/Private/AssetFixUpRedirectors.cpp b/Engine/Source/Developer/AssetTools/Private/AssetFixUpRedirectors.cpp
--- a/Engine/Source/Developer/AssetTools/Private/AssetFixUpRedirectors.cpp
+++ b/Engine/Source/Developer/AssetTools/Private/AssetFixUpRedirectors.cpp
@@ -852,10 +852,9 @@ void FAssetFixUpRedirectors::ExecuteFixUp(TArray<TWeakObjectPtr<UObjectRedirecto
 			for (const UObjectRedirector* Redirector = RedirectorRefs.Redirector.Get(); Redirector; Redirector = Cast<UObjectRedirector>(Redirector->DestinationObject))
 			{
 				const FSoftObjectPath RedirectorObjectPath = FSoftObjectPath(Redirector);
-				FText Error;
-				if (!CollectionManagerModule.Get().HandleRedirectorDeleted(RedirectorObjectPath, &Error))
+				if (!CollectionManagerModule.Get().HandleRedirectorDeleted(RedirectorObjectPath))
 				{
-					RedirectorRefs.OtherFailures.Add(FText::Format(LOCTEXT("RedirectorFixupFailed_CollectionsFailedToSave", "Referencing collection(s) failed to save: {0}"), Error));
+					RedirectorRefs.OtherFailures.Add(FText::Format(LOCTEXT("RedirectorFixupFailed_CollectionsFailedToSave", "Referencing collection(s) failed to save: {0}"), CollectionManagerModule.Get().GetLastError()));
 				}
 			}
 		}
