Updated ExportAssets to include option for using CleanFileName as the file name

This changelist adds a new function to asset tools in order to export assets using their clean file name. Without this change, the assets are exported to their desired destination -- but the asset is copied over nested in it's entire existing path

I also added this as an entirely new method in order to avoid changing all of the other call sites where `ExportAssets` was being used.

#tests
-Ran editor locally, confirmed using this new method exported assets using the new function

#rnx

diff --git a/Engine/Source/Developer/AssetTools/Private/AssetTools.cpp b/Engine/Source/Developer/AssetTools/Private/AssetTools.cpp
--- a/Engine/Source/Developer/AssetTools/Private/AssetTools.cpp
+++ b/Engine/Source/Developer/AssetTools/Private/AssetTools.cpp
@@ -2917,18 +2917,28 @@ void UAssetToolsImpl::ExportAssets(const TArray<FString>& AssetsToExport, const
 	}
 
 	const bool bPromptIndividualFilenames = false;
-	ExportAssetsInternal(AssetObjectsToExport, bPromptIndividualFilenames, ExportPath);
+	const bool bExportAsCleanFileName = false;
+	ExportAssetsInternal(AssetObjectsToExport, bPromptIndividualFilenames, bExportAsCleanFileName, ExportPath);
 }
 
 void UAssetToolsImpl::ExportAssets(const TArray<UObject*>& AssetsToExport, const FString& ExportPath) const
 {
 	const bool bPromptIndividualFilenames = false;
-	ExportAssetsInternal(AssetsToExport, bPromptIndividualFilenames, ExportPath);
+	const bool bSaveAsCleanFilename = false;
+	ExportAssetsInternal(AssetsToExport, bPromptIndividualFilenames, bSaveAsCleanFilename, ExportPath);
 }
 
+void UAssetToolsImpl::ExportAssetsWithCleanFilename(const TArray<UObject*>& AssetsToExport, const FString& ExportPath) const
+{
+	const bool bPromptIndividualFilenames = false;
+	const bool bSaveAsCleanFilename = true;
+	ExportAssetsInternal(AssetsToExport, bPromptIndividualFilenames, bSaveAsCleanFilename, ExportPath);
+}
+ 
 void UAssetToolsImpl::ExportAssetsWithDialog(const TArray<UObject*>& AssetsToExport, bool bPromptForIndividualFilenames)
 {
-	ExportAssetsInternal(AssetsToExport, bPromptForIndividualFilenames, TEXT(""));
+	const bool bSaveAsCleanFilename = false;
+	ExportAssetsInternal(AssetsToExport, bPromptForIndividualFilenames, bSaveAsCleanFilename, TEXT(""));
 }
 
 void UAssetToolsImpl::ExportAssetsWithDialog(const TArray<FString>& AssetsToExport, bool bPromptForIndividualFilenames)
@@ -2949,7 +2959,8 @@ void UAssetToolsImpl::ExportAssetsWithDialog(const TArray<FString>& AssetsToExpo
 		}
 	}
 
-	ExportAssetsInternal(AssetObjectsToExport, bPromptForIndividualFilenames, TEXT(""));
+	const bool bSaveAsCleanFilename = false;
+	ExportAssetsInternal(AssetObjectsToExport, bPromptForIndividualFilenames, bSaveAsCleanFilename, TEXT(""));
 }
 
 void UAssetToolsImpl::ExpandDirectories(const TArray<FString>& Files, const FString& DestinationPath, TArray<TPair<FString, FString>>& FilesAndDestinations) const
@@ -4218,7 +4229,7 @@ TArray<UObject*> UAssetToolsImpl::ImportAssetsInternal(const TArray<FString>& Fi
 	return ReturnObjects;
 }
 
-void UAssetToolsImpl::ExportAssetsInternal(const TArray<UObject*>& ObjectsToExport, bool bPromptIndividualFilenames, const FString& ExportPath) const
+void UAssetToolsImpl::ExportAssetsInternal(const TArray<UObject*>& ObjectsToExport, bool bPromptIndividualFilenames, bool bSaveAsCleanFilename, const FString& ExportPath) const
 {
 	FString LastExportPath = !ExportPath.IsEmpty() ? ExportPath : FEditorDirectories::Get().GetLastDirectory(ELastDirectory::GENERIC_EXPORT);
 
@@ -4476,7 +4487,8 @@ void UAssetToolsImpl::ExportAssetsInternal(const TArray<UObject*>& ObjectsToExpo
 				}
 
 				FPaths::NormalizeFilename(PackageName);
-				SaveFileName /= PackageName;
+				
+				SaveFileName /= bSaveAsCleanFilename ? FPaths::GetCleanFilename(PackageName) : PackageName;
 			}
 			else
 			{
diff --git a/Engine/Source/Developer/AssetTools/Private/AssetTools.h b/Engine/Source/Developer/AssetTools/Private/AssetTools.h
--- a/Engine/Source/Developer/AssetTools/Private/AssetTools.h
+++ b/Engine/Source/Developer/AssetTools/Private/AssetTools.h
@@ -103,6 +103,7 @@ public:
 	virtual void ImportAssetTasks(const TArray<UAssetImportTask*>& ImportTasks) override;
 	virtual void ExportAssets(const TArray<FString>& AssetsToExport, const FString& ExportPath) override;
 	virtual void ExportAssets(const TArray<UObject*>& AssetsToExport, const FString& ExportPath) const override;
+	virtual void ExportAssetsWithCleanFilename(const TArray<UObject*>& AssetsToExport, const FString& ExportPath) const override;
 	virtual void ExportAssetsWithDialog(const TArray<UObject*>& AssetsToExport, bool bPromptForIndividualFilenames) override;
 	virtual void ExportAssetsWithDialog(const TArray<FString>& AssetsToExport, bool bPromptForIndividualFilenames) override;
 	virtual bool CanExportAssets(const TArray<FAssetData>& AssetsToExport) const override;
@@ -207,8 +208,8 @@ private:
 	/** Internal method that performs the actual asset importing */
 	TArray<UObject*> ImportAssetsInternal(const TArray<FString>& Files, const FString& RootDestinationPath, TArray<TPair<FString, FString>> *FilesAndDestinationsPtr, const FAssetImportParams& ImportParams) const;
 
-	/** Internal method to export assets.  If no export path is created a user will be prompted for one.  if bPromptIndividualFilenames is true a user will be asked per file */
-	void ExportAssetsInternal(const TArray<UObject*>& ObjectsToExport, bool bPromptIndividualFilenames, const FString& ExportPath) const;
+	/** Internal method to export assets.  If no export path is created a user will be prompted for one.  if bPromptIndividualFilenames is true a user will be asked per file. if bSaveAsCleanFilename is true the file will be saved as the clean filename */
+	void ExportAssetsInternal(const TArray<UObject*>& ObjectsToExport, bool bPromptIndividualFilenames, bool bSaveAsCleanFilename, const FString& ExportPath) const;
 
 	UObject* PerformDuplicateAsset(const FString& AssetName, const FString& PackagePath, UObject* OriginalObject, bool bWithDialog);
 
diff --git a/Engine/Source/Developer/AssetTools/Public/IAssetTools.h b/Engine/Source/Developer/AssetTools/Public/IAssetTools.h
--- a/Engine/Source/Developer/AssetTools/Public/IAssetTools.h
+++ b/Engine/Source/Developer/AssetTools/Public/IAssetTools.h
@@ -464,6 +464,14 @@ public:
 	 */
 	virtual void ExportAssets(const TArray<UObject*>& AssetsToExport, const FString& ExportPath) const = 0;
 	
+	/**
+	 * Exports the specified objects to file using the clean filename as the saved filename.
+	 *
+	 * @param	AssetsToExport					List of assets to export
+	 * @param	ExportPath						The directory path to export to.
+	 */
+	virtual void ExportAssetsWithCleanFilename(const TArray<UObject*>& AssetsToExport, const FString& ExportPath) const = 0;
+
 	/**
 	 * Exports the specified objects to file. First prompting the user to pick an export directory and optionally prompting the user to pick a unique directory per file
 	 *
