When importing an asset never force a content browser sync if the caller didn't want a sync.

This fix an issue where the import of an asset pack would cause one of the CB to loose its setting because the tab was open to early in the engine boot cycle.

diff --git a/Engine/Source/Developer/AssetTools/Private/AssetTools.cpp b/Engine/Source/Developer/AssetTools/Private/AssetTools.cpp
--- a/Engine/Source/Developer/AssetTools/Private/AssetTools.cpp
+++ b/Engine/Source/Developer/AssetTools/Private/AssetTools.cpp
@@ -3671,8 +3671,6 @@ TArray<UObject*> UAssetToolsImpl::ImportAssetsInternal(const TArray<FString>& Fi
 
 	TSharedPtr<FInterchangeImportStatus, ESPMode::ThreadSafe> ImportStatus = MakeShared<FInterchangeImportStatus>(FilesAndDestinations.Num());
 
-	const bool bForceContentBrowserSyncIfOnlyOneMainAsset = (FilesAndDestinations.Num() == 1 && !bAutomatedImport);
-
 	// Now iterate over the input files and use the same factory object for each file with the same extension
 	for(int32 FileIdx = 0; FileIdx < FilesAndDestinations.Num() && !bImportWasCancelled; ++FileIdx)
 	{
@@ -3751,7 +3749,7 @@ TArray<UObject*> UAssetToolsImpl::ImportAssetsInternal(const TArray<FString>& Fi
 				
 				TFunction<void(UE::Interchange::FImportResult&)> AppendAndBroadcastImportResultIfNeeded =
 					// Note: ImportStatus captured by value so that the lambda keeps the shared ptr alive
-					[bSceneImport = Params.bSceneImport, ImportStatus, AppendImportResult, bSyncToBrowser, bForceContentBrowserSyncIfOnlyOneMainAsset](UE::Interchange::FImportResult& Result)
+					[bSceneImport = Params.bSceneImport, ImportStatus, AppendImportResult, bSyncToBrowser](UE::Interchange::FImportResult& Result)
 					{
 						AppendImportResult(Result);
 
@@ -3793,7 +3791,7 @@ TArray<UObject*> UAssetToolsImpl::ImportAssetsInternal(const TArray<FString>& Fi
 							}
 
 							//Force browser to sync to the import asset if there is only one asset imported
-							if (bSyncToBrowser || (bForceContentBrowserSyncIfOnlyOneMainAsset && MainAssets.Num() == 1))
+							if (bSyncToBrowser)
 							{
 								UAssetToolsImpl::Get().SyncBrowserToAssets(MainAssets);
 							}
@@ -4221,7 +4219,7 @@ TArray<UObject*> UAssetToolsImpl::ImportAssetsInternal(const TArray<FString>& Fi
 	}
 
 	// Sync content browser to the newly created assets
-	if(ReturnObjects.Num() && (bSyncToBrowser != false || (ReturnObjects.Num() == 1 && bForceContentBrowserSyncIfOnlyOneMainAsset)))
+	if(ReturnObjects.Num() && bSyncToBrowser)
 	{
 		UAssetToolsImpl::Get().SyncBrowserToAssets(ReturnObjects);
 	}
