Iterate over migrated packags and do WorldCleanup() on found world. This is necessary to avoid various issues iwht Unitialization() in the BeginDestroys() of world subsystems.

diff --git a/Engine/Source/Developer/AssetTools/Private/AssetTools.cpp b/Engine/Source/Developer/AssetTools/Private/AssetTools.cpp
--- a/Engine/Source/Developer/AssetTools/Private/AssetTools.cpp
+++ b/Engine/Source/Developer/AssetTools/Private/AssetTools.cpp
@@ -1160,6 +1160,24 @@ namespace UE::AssetTools::Private
 			// Turn off the components while unloading stuff
 			FGlobalComponentReregisterContext ComponentContext;
 
+			// 1st Force to clean all the worlds so they can shutdown the subsystems properly
+			for (const TWeakObjectPtr<UPackage>& WeakPackage : PackagesToClean)
+			{
+				if (UPackage* Package = WeakPackage.Get())
+				{
+					if (Package->ContainsMap())
+					{
+						ForEachObjectWithOuter(Package, [](UObject* Object)
+						{
+							if (UWorld* World = Cast<UWorld>(Object))
+							{
+								World->CleanupWorld();
+							}
+						});
+					}
+				}
+			}
+
 			TArray<UObject*> ReferenceToNull;
 
 			// We do the clean pass of the packages in two loop because the PurgeObject can affect the ability to get the main object from another package.
