Unregister AssetDefinitions from registry
* Sync function was not taking into account unregistered AssetDefinitions - the fix is to iterate the AssetTypeActions also and find corresponding AssetDefinitions and remove if there is none.
* SAssetFilterBar caches AssetDefinitions - added a notification for when the AssetDefinitionRegistry changes.  Put the broadcast on a one-shot ticker instead of directly calling as the Register/Unregister can occur frequently.  Breaks up the callback depth.
* Added a "Version" to AssetDefinitionRegistry as clients were using the number of assets as a version. This can fall over if there is an Add and Remove between two calls to GetAllAssetDefinitions.

diff --git a/Engine/Source/Developer/AssetTools/Private/AssetTools.cpp b/Engine/Source/Developer/AssetTools/Private/AssetTools.cpp
--- a/Engine/Source/Developer/AssetTools/Private/AssetTools.cpp
+++ b/Engine/Source/Developer/AssetTools/Private/AssetTools.cpp
@@ -1453,14 +1453,15 @@ void UAssetToolsImpl::RegisterAssetTypeActions(const TSharedRef<IAssetTypeAction
 	}
 }
 
-void UAssetToolsImpl::SyncAssetTypesToAssetDefinitions() const
+void UAssetToolsImpl::SyncAssetTypesToAssetDefinitions()
 {
-	static int32 CachedAmount = 0;
-
-	TArray<TObjectPtr<UAssetDefinition>> AssetDefinitions = UAssetDefinitionRegistry::Get()->GetAllAssetDefinitions();
-	if (CachedAmount != AssetDefinitions.Num())
+	uint64 Version = UAssetDefinitionRegistry::Get()->GetAssetDefinitionVersion();
+	if (CachedVersion == 0 || CachedVersion != Version)
 	{
-		CachedAmount = AssetDefinitions.Num();		
+		CachedVersion = Version;
+		
+		TArray<TObjectPtr<UAssetDefinition>> AssetDefinitions = UAssetDefinitionRegistry::Get()->GetAllAssetDefinitions();
+		
 		for (UAssetDefinition* AssetDefinition : AssetDefinitions)
 		{
 			if (!AssetDefinition->IsA<UAssetDefinition_AssetTypeActionsProxy>())
@@ -1473,10 +1474,39 @@ void UAssetToolsImpl::SyncAssetTypesToAssetDefinitions() const
 				{
 					TSharedRef<FAssetDefinitionProxy> Proxy = MakeShared<FAssetDefinitionProxy>(AssetDefinition);
 					// Cache the asset definition proxy.
-					const_cast<UAssetToolsImpl*>(this)->RegisterAssetTypeActions(Proxy);
+					RegisterAssetTypeActions(Proxy);
 				}
 			}
 		}
+
+		// Remove any AssetDefinitionProxy objects that don't have an associated AssetDefinition
+		for (int32 Index = AssetTypeActionsList.Num() - 1, Last = 0; Index >= Last; --Index)
+		{
+			const TSharedRef<IAssetTypeActions>& AssetTypeActions = AssetTypeActionsList[Index];
+			if (AssetTypeActions->IsAssetDefinitionInDisguise())
+			{
+				UClass* AssetTypeActionsSupportedClass = AssetTypeActions->GetSupportedClass();
+				
+				bool bFound = false;
+				// Find it in the AssetDefinitions List
+				for (UAssetDefinition* AssetDefinition : AssetDefinitions)
+				{
+					if (!AssetDefinition->IsA<UAssetDefinition_AssetTypeActionsProxy>())
+					{
+						const TSoftClassPtr<UObject> AssetClass = AssetDefinition->GetAssetClass();
+						if (AssetClass == AssetTypeActionsSupportedClass)
+						{
+							bFound = true;
+							break;
+						}
+					}
+				}
+				if (!bFound)
+				{
+					RemoveAssetTypeActionBySupportedClass(AssetTypeActionsSupportedClass);
+				}
+			}
+		}		
 	}
 }
 
@@ -1501,7 +1531,7 @@ void UAssetToolsImpl::UnregisterAssetTypeActions(const TSharedRef<IAssetTypeActi
 
 void UAssetToolsImpl::GetAssetTypeActionsList( TArray<TWeakPtr<IAssetTypeActions>>& OutAssetTypeActionsList ) const
 {
-	SyncAssetTypesToAssetDefinitions();
+	const_cast<UAssetToolsImpl*>(this)->SyncAssetTypesToAssetDefinitions();
 	for (auto ActionsIt = AssetTypeActionsList.CreateConstIterator(); ActionsIt; ++ActionsIt)
 	{
 		OutAssetTypeActionsList.Add(*ActionsIt);
@@ -1510,7 +1540,7 @@ void UAssetToolsImpl::GetAssetTypeActionsList( TArray<TWeakPtr<IAssetTypeActions
 
 TWeakPtr<IAssetTypeActions> UAssetToolsImpl::GetAssetTypeActionsForClass(const UClass* Class) const
 {
-	SyncAssetTypesToAssetDefinitions();
+	const_cast<UAssetToolsImpl*>(this)->SyncAssetTypesToAssetDefinitions();
 
 	const UClass* CandidateClass = Class;
 	
@@ -1568,7 +1598,7 @@ TArray<TWeakPtr<IAssetTypeActions>> UAssetToolsImpl::GetAssetTypeActionsListForC
 {
 	TArray<TWeakPtr<IAssetTypeActions>> ResultAssetTypeActionsList;
 
-	SyncAssetTypesToAssetDefinitions();
+	const_cast<UAssetToolsImpl*>(this)->SyncAssetTypesToAssetDefinitions();
 	for (int32 TypeActionsIdx = 0; TypeActionsIdx < AssetTypeActionsList.Num(); ++TypeActionsIdx)
 	{
 		TSharedRef<IAssetTypeActions> TypeActions = AssetTypeActionsList[TypeActionsIdx];
diff --git a/Engine/Source/Developer/AssetTools/Private/AssetTools.h b/Engine/Source/Developer/AssetTools/Private/AssetTools.h
--- a/Engine/Source/Developer/AssetTools/Private/AssetTools.h
+++ b/Engine/Source/Developer/AssetTools/Private/AssetTools.h
@@ -226,12 +226,14 @@ private:
 	TArray<UObject*> ImportAssetsWithDialogImplementation(const FString& DestinationPath, bool bAllowAsyncImport);
 
 	/** Make sure we're not syncing */
-	void SyncAssetTypesToAssetDefinitions() const;
+	void SyncAssetTypesToAssetDefinitions();
 
 	/** Helper to remove an entry AssetTypeActionsList given a class returned by GetSupportedClass */
 	void RemoveAssetTypeActionBySupportedClass(const UClass* SupportedClass);
 
 private:
+	uint64 CachedVersion = 0;
+	
 	/** The list of all registered AssetTypeActions */
 	TArray<TSharedRef<IAssetTypeActions>> AssetTypeActionsList;
 
diff --git a/Engine/Source/Editor/AssetDefinition/Private/AssetDefinitionRegistry.cpp b/Engine/Source/Editor/AssetDefinition/Private/AssetDefinitionRegistry.cpp
--- a/Engine/Source/Editor/AssetDefinition/Private/AssetDefinitionRegistry.cpp
+++ b/Engine/Source/Editor/AssetDefinition/Private/AssetDefinitionRegistry.cpp
@@ -24,6 +24,11 @@ UAssetDefinitionRegistry* UAssetDefinitionRegistry::Get()
 	return Singleton;
 }
 
+UAssetDefinitionRegistry::UAssetDefinitionRegistry()
+{
+	TickerDelegate = FTickerDelegate::CreateUObject(this, &UAssetDefinitionRegistry::TickVersionNotification);
+}
+
 void UAssetDefinitionRegistry::BeginDestroy()
 {
 	if (Singleton == this)
@@ -63,6 +68,11 @@ TArray<TObjectPtr<UAssetDefinition>> UAssetDefinitionRegistry::GetAllAssetDefini
 	return AllAssetDefinitions;
 }
 
+uint64 UAssetDefinitionRegistry::GetAssetDefinitionVersion() const
+{
+	return Version;
+}
+
 void UAssetDefinitionRegistry::RegisterAssetDefinition(UAssetDefinition* AssetDefinition)
 {
 	check(AssetDefinition);
@@ -79,6 +89,9 @@ void UAssetDefinitionRegistry::RegisterAssetDefinition(UAssetDefinition* AssetDe
 		if (CanAddDefinition)
 		{
 			AssetDefinitions.Add(SupportedClass, AssetDefinition);
+
+			Version++;
+			RegisterTickerForVersionNotification();
 		}
 		else
 		{
@@ -95,5 +108,28 @@ void UAssetDefinitionRegistry::UnregisterAssetDefinition(UAssetDefinition* Asset
 	if (TSoftClassPtr<UObject> SupportedClass = AssetDefinition->GetAssetClass())
 	{
 		AssetDefinitions.Remove(SupportedClass);
+		
+		Version++;
+		RegisterTickerForVersionNotification();
+	}
+}
+
+FOnAssetDefinitionRegistryVersionChange& UAssetDefinitionRegistry::OnAssetDefinitionRegistryVersionChange()
+{
+	return OnAssetDefinitionRegistryVersionChangeDelegate;
+}
+
+void UAssetDefinitionRegistry::RegisterTickerForVersionNotification()
+{
+	if (!TickerDelegateHandle.IsValid())
+	{
+		TickerDelegateHandle = FTSTicker::GetCoreTicker().AddTicker(TickerDelegate);
 	}
-}
\ No newline at end of file
+}
+
+bool UAssetDefinitionRegistry::TickVersionNotification(float)
+{
+	OnAssetDefinitionRegistryVersionChangeDelegate.Broadcast(this);
+	TickerDelegateHandle.Reset();
+	return false; // One-shot
+}
diff --git a/Engine/Source/Editor/AssetDefinition/Public/AssetDefinitionRegistry.h b/Engine/Source/Editor/AssetDefinition/Public/AssetDefinitionRegistry.h
--- a/Engine/Source/Editor/AssetDefinition/Public/AssetDefinitionRegistry.h
+++ b/Engine/Source/Editor/AssetDefinition/Public/AssetDefinitionRegistry.h
@@ -2,25 +2,33 @@
 
 #pragma once
 
+#include "Containers/Ticker.h"
+#include "Engine/EngineBaseTypes.h"
 #include "UObject/Object.h"
 
 #include "AssetDefinitionRegistry.generated.h"
 
 class UAssetDefinition;
 
+DECLARE_MULTICAST_DELEGATE_OneParam(FOnAssetDefinitionRegistryVersionChange, UAssetDefinitionRegistry*);
+
 UCLASS(config=Editor)
 class ASSETDEFINITION_API UAssetDefinitionRegistry : public UObject
 {
 	GENERATED_BODY()
 
-public:
+public:	
 	static UAssetDefinitionRegistry* Get();
+
+	UAssetDefinitionRegistry();
 	
 	virtual void BeginDestroy() override;
 
 	const UAssetDefinition* GetAssetDefinitionForAsset(const FAssetData& Asset) const;
 	const UAssetDefinition* GetAssetDefinitionForClass(const UClass* Class) const;
 
+	// Gets the current version of the AssetDefinitions.   Version is updated whenever an AssetDefinition is Registered/Unregistered
+	uint64 GetAssetDefinitionVersion() const;
 	TArray<TObjectPtr<UAssetDefinition>> GetAllAssetDefinitions() const;
 
 	/**
@@ -32,11 +40,25 @@ public:
 	void RegisterAssetDefinition(UAssetDefinition* AssetDefinition);
 
 	void UnregisterAssetDefinition(UAssetDefinition* AssetDefinition);
+
+	/**
+	 * Called when the AssetDefinitionRegistry's version has changed.
+	 */
+	FOnAssetDefinitionRegistryVersionChange& OnAssetDefinitionRegistryVersionChange();
 	
 private:
+	void RegisterTickerForVersionNotification();
+	bool TickVersionNotification(float);
+	
 	static UAssetDefinitionRegistry* Singleton;
 	static bool bHasShutDown;
 
 	UPROPERTY()
 	TMap<TSoftClassPtr<UObject>, TObjectPtr<UAssetDefinition>> AssetDefinitions;
+
+	uint64 Version;
+
+	FTickerDelegate TickerDelegate;
+	FTSTicker::FDelegateHandle TickerDelegateHandle;
+	FOnAssetDefinitionRegistryVersionChange OnAssetDefinitionRegistryVersionChangeDelegate;
 };
\ No newline at end of file
diff --git a/Engine/Source/Editor/EditorWidgets/Public/Filters/SAssetFilterBar.h b/Engine/Source/Editor/EditorWidgets/Public/Filters/SAssetFilterBar.h
--- a/Engine/Source/Editor/EditorWidgets/Public/Filters/SAssetFilterBar.h
+++ b/Engine/Source/Editor/EditorWidgets/Public/Filters/SAssetFilterBar.h
@@ -242,7 +242,7 @@ public:
 
 		// Re-create the Asset Type Action filters whenever the permission list changes
 		AssetClassPermissionList->OnFilterChanged().AddSP(this, &SAssetFilterBar<FilterType>::CreateAssetTypeActionFilters);
-
+		UAssetDefinitionRegistry::Get()->OnAssetDefinitionRegistryVersionChange().AddLambda([this](UAssetDefinitionRegistry*) { CreateAssetTypeActionFilters(); });
 	}
 
 	virtual void SaveSettings()
