Fixed a crash when canceling an asset migration. When moving back the packages to their original path there could be some packages left from the migration.

diff --git a/Engine/Source/Developer/AssetTools/Private/AssetTools.cpp b/Engine/Source/Developer/AssetTools/Private/AssetTools.cpp
--- a/Engine/Source/Developer/AssetTools/Private/AssetTools.cpp
+++ b/Engine/Source/Developer/AssetTools/Private/AssetTools.cpp
@@ -958,7 +958,7 @@ namespace UE::AssetTools::Private
 			}
 		}
 
-		static void CreatePackagesAndSetupLinkers(FPackageMigrationContext& PackageMigrationContext, FPackageMigrationImplContext& MigrationImplContext, TArray<TWeakObjectPtr<UPackage>>& PackagesToClean)
+		static void CreatePackagesAndSetupLinkers(FPackageMigrationContext& PackageMigrationContext, FPackageMigrationImplContext& MigrationImplContext, TArray<FName>& PackagesToClean)
 		{
 			FArchive* ReaderOverride = nullptr;
 
@@ -968,23 +968,25 @@ namespace UE::AssetTools::Private
 				if (MigrationImplContext.SlowTask.ShouldCancel() || MigrationImplContext.bWasCanceled)
 				{
 					MigrationImplContext.bWasCanceled = true;
-					break;
+					PackagesToClean.Add(FName(*MigrationPackageData.GetInstancedPackageName()));
 				}
+				else
+				{
+					UPackage* MigrationPackage = CreatePackage(*MigrationPackageData.GetInstancedPackageName());
 
-				UPackage* MigrationPackage = CreatePackage(*MigrationPackageData.GetInstancedPackageName());
-
-				/**
-				 * Load_Verify tell the linker to not load the package but it will create the linker if the file exist and is valid.
-				 * LOAD_NoVerify tell the linker to not check the import of the package (in the editor this will avoid loading the hard dependencies of the package)
-				 * 
-				 * This will the package with a pre-created linker that has the right instancing context. Ready to loaded by another call to load package that may not have all that info.
-				 */
-				MigrationPackageData.InstancedPackage = LoadPackage(MigrationPackage, *MigrationPackageData.GetOriginalPackageName(), LOAD_Verify | LOAD_NoVerify, ReaderOverride, &MigrationImplContext.InstancingContext);
+					/**
+					 * Load_Verify tell the linker to not load the package but it will create the linker if the file exist and is valid.
+					 * LOAD_NoVerify tell the linker to not check the import of the package (in the editor this will avoid loading the hard dependencies of the package)
+					 * 
+					 * This will the package with a pre-created linker that has the right instancing context. Ready to loaded by another call to load package that may not have all that info.
+					 */
+					MigrationPackageData.InstancedPackage = LoadPackage(MigrationPackage, *MigrationPackageData.GetOriginalPackageName(), LOAD_Verify | LOAD_NoVerify, ReaderOverride, &MigrationImplContext.InstancingContext);
 
-				// Add the PKG_NewlyCreated flag here to help the engine code that expect the packages without it to have a file associated with it on disk.
-				MigrationPackage->SetPackageFlags(PKG_NewlyCreated);
+					// Add the PKG_NewlyCreated flag here to help the engine code that expect the packages without it to have a file associated with it on disk.
+					MigrationPackage->SetPackageFlags(PKG_NewlyCreated);
 
-				PackagesToClean.Add(MigrationPackage);
+					PackagesToClean.Add(MigrationPackage->GetFName());
+				}
 			}
 		}
 
@@ -1122,7 +1124,7 @@ namespace UE::AssetTools::Private
 			}
 		}
 
-		static void CleanInstancedPackages(const TArray<TWeakObjectPtr<UPackage>>& PackagesToClean, FPackageMigrationContext& PackageMigrationContext)
+		static void CleanInstancedPackages(const TArray<FName>& PackagesToClean, FPackageMigrationContext& PackageMigrationContext)
 		{
 			for (FPackageMigrationContext::FMigrationPackageData& MigrationPackageData : PackageMigrationContext.MigrationPackagesData)
 			{
@@ -1160,48 +1162,52 @@ namespace UE::AssetTools::Private
 			// Turn off the components while unloading stuff
 			FGlobalComponentReregisterContext ComponentContext;
 
+
+			TArray<UPackage*> ExistingPackageToClean;
+			ExistingPackageToClean.Reserve(PackagesToClean.Num());
+
+			for (const FName PackageName : PackagesToClean)
+			{
+				if (UPackage* Package = FindObjectFast<UPackage>(nullptr, PackageName, true))
+				{
+					ExistingPackageToClean.Add(Package);
+				}
+			}
+
 			// 1st Force to clean all the worlds so they can shutdown the subsystems properly
-			for (const TWeakObjectPtr<UPackage>& WeakPackage : PackagesToClean)
+			for (const UPackage* Package : ExistingPackageToClean)
 			{
-				if (UPackage* Package = WeakPackage.Get())
+				if (Package->ContainsMap())
 				{
-					if (Package->ContainsMap())
-					{
-						ForEachObjectWithOuter(Package, [](UObject* Object)
+					ForEachObjectWithOuter(Package, [](UObject* Object)
 						{
 							if (UWorld* World = Cast<UWorld>(Object))
 							{
 								World->CleanupWorld();
 							}
 						});
-					}
 				}
 			}
 
 			TArray<UObject*> ReferenceToNull;
 
+
 			// We do the clean pass of the packages in two loop because the PurgeObject can affect the ability to get the main object from another package.
-			for (const TWeakObjectPtr<UPackage>& WeakPackage : PackagesToClean)
+			for (UPackage* Package : ExistingPackageToClean)
 			{
-				if (UPackage* Package = WeakPackage.Get())
+				if (UObject* Asset = FPackageMigrationImpl::FindAssetInPackage(Package))
 				{
-					if (UObject* Asset = FPackageMigrationImpl::FindAssetInPackage(Package))
-					{
-						AssetRegistry.AssetDeleted(Asset);
-						ReferenceToNull.Add(Asset);
-					}
+					AssetRegistry.AssetDeleted(Asset);
+					ReferenceToNull.Add(Asset);
 				}
 			}
 
-			for (const TWeakObjectPtr<UPackage>& WeakPackage : PackagesToClean)
+			for (UPackage* Package : ExistingPackageToClean)
 			{
-				if (UPackage* Package = WeakPackage.Get())
-				{
-					const ERenameFlags PkgRenameFlags = REN_ForceNoResetLoaders | REN_DoNotDirty | REN_DontCreateRedirectors | REN_NonTransactional | REN_SkipGeneratedClasses;
-					check(Package->Rename(*MakeUniqueObjectName(nullptr, UPackage::StaticClass(), *FString::Printf(TEXT("%s_DEADFROMMIGRATION"), *Package->GetName())).ToString(), nullptr, PkgRenameFlags));
-					(*PurgeObject)(Package);
-					ForEachObjectWithOuter(Package, *PurgeObject);
-				}
+				const ERenameFlags PkgRenameFlags = REN_ForceNoResetLoaders | REN_DoNotDirty | REN_DontCreateRedirectors | REN_NonTransactional | REN_SkipGeneratedClasses;
+				check(Package->Rename(*MakeUniqueObjectName(nullptr, UPackage::StaticClass(), *FString::Printf(TEXT("%s_DEADFROMMIGRATION"), *Package->GetName())).ToString(), nullptr, PkgRenameFlags));
+				(*PurgeObject)(Package);
+				ForEachObjectWithOuter(Package, *PurgeObject);
 			}
 
 			// Removing assets from memory is complicated
@@ -5236,7 +5242,7 @@ void UAssetToolsImpl::MigratePackages_ReportConfirmed(TSharedPtr<TArray<ReportPa
 
 
 			// 3) Instanced load
-			TArray<TWeakObjectPtr<UPackage>> PackagesToClean;
+			TArray<FName> PackagesToClean;
 			PackagesToClean.Reserve(PackageMigrationContext.MigrationPackagesData.Num());
 			FPackageMigrationImpl::CreatePackagesAndSetupLinkers(PackageMigrationContext, MigrationImplContext, PackagesToClean);
 
