Adding Multi root remapping to header patcher for Duplicate plugin.

diff --git a/Engine/Source/Developer/AssetTools/Private/AssetHeaderPatcher.cpp b/Engine/Source/Developer/AssetTools/Private/AssetHeaderPatcher.cpp
--- a/Engine/Source/Developer/AssetTools/Private/AssetHeaderPatcher.cpp
+++ b/Engine/Source/Developer/AssetTools/Private/AssetHeaderPatcher.cpp
@@ -606,6 +606,33 @@ bool FAssetHeaderPatcherInner::DoPatch(FString& InOutString)
 		}
 	}
 
+	{
+		// Patching embedded mount points.
+		// 
+		FStringView InView = InOutString;
+		FStringView MountPoint = FPathViews::GetMountPointNameFromPath(InView);
+		if (MountPoint.Len())
+		{
+			TStringBuilder<NAME_SIZE> Tmp;
+			Tmp.Append(TEXT("<Mountpoint>"));
+			Tmp.Append(MountPoint);
+			FStringView MaybeReplacement = Find(SearchAndReplace, FStringView(Tmp));
+			if (!MaybeReplacement.IsEmpty())
+			{
+				int32 Left = UE_PTRDIFF_TO_INT32(MountPoint.GetData() - InView.GetData());
+				int32 Right = Left + MountPoint.Len();
+
+				Tmp.Reset();
+				Tmp.Append(InView.Left(Left));
+				Tmp.Append(MaybeReplacement);
+				Tmp.Append(InView.RightChop(Right));
+
+				InOutString = Tmp;
+				return true;
+			}
+		}
+	}
+
 	return false;
 }
 
diff --git a/Engine/Source/Developer/AssetTools/Private/AssetTools.cpp b/Engine/Source/Developer/AssetTools/Private/AssetTools.cpp
--- a/Engine/Source/Developer/AssetTools/Private/AssetTools.cpp
+++ b/Engine/Source/Developer/AssetTools/Private/AssetTools.cpp
@@ -2791,10 +2791,11 @@ TMap<FString, FString> UAssetToolsImpl::GetMappingsForRootPackageRename(
 	const FString& SrcRoot,
 	const FString& DstRoot,
 	const FString& SrcBaseDir,
-	const TArray<TPair<FString, FString>>& SourceAndDestFiles) const
+	const TArray<TPair<FString, FString>>& SourceAndDestFiles,
+	const TMap<FString, FString>& MountPointReplacements) const
 {
 	TMap<FString, FString> Result;
-	Result.Reserve(6 + SourceAndDestFiles.Num()); // usually only 3, +6 just in case
+	Result.Reserve(6 + SourceAndDestFiles.Num() + MountPointReplacements.Num()); // usually only 3, +6 just in case
 
 	{	// Plugin name patterns
 		FString SrcPath = FPaths::Combine(TEXT("/"), SrcRoot, SrcRoot);
@@ -2843,6 +2844,16 @@ TMap<FString, FString> UAssetToolsImpl::GetMappingsForRootPackageRename(
 		}
 	}
 
+	// Mount point replacements.
+	// In the case where a project has multiple plugins, they are represented as different package roots (mountpoints)
+	// So a Package my reference an asset, or object in a different mount point.
+	// We tag this information with '<Mountpoint>' so the header patcher can do a partial replacement (instead of the normal complete replacement)
+	// if the mount point is detected.
+	for (const TPair<FString, FString>& MountPair : MountPointReplacements)
+	{
+		Result.Add(TEXT("<Mountpoint>") + MountPair.Key, MountPair.Value);
+	}
+
 	return Result;
 }
 
diff --git a/Engine/Source/Developer/AssetTools/Private/AssetTools.h b/Engine/Source/Developer/AssetTools/Private/AssetTools.h
--- a/Engine/Source/Developer/AssetTools/Private/AssetTools.h
+++ b/Engine/Source/Developer/AssetTools/Private/AssetTools.h
@@ -125,7 +125,7 @@ public:
 	virtual bool AdvancedCopyPackages(const FAdvancedCopyParams& CopyParams, const TArray<TMap<FString, FString>>& PackagesAndDestinations) const override;
 	virtual bool AdvancedCopyPackages(const TMap<FString, FString>& SourceAndDestPackages, const bool bForceAutosave, const bool bCopyOverAllDestinationOverlaps, FDuplicatedObjects* OutDuplicatedObjects, EMessageSeverity::Type NotificationSeverityFilter) const override;
 	virtual bool PatchCopyPackageFile(const FString& SrcFile, const FString& DstFile, const TMap<FString, FString>& SearchForAndReplace) const;
-	virtual TMap<FString, FString> GetMappingsForRootPackageRename(const FString& SrcRoot, const FString& DstRoot, const FString& SrcBaseDir, const TArray<TPair<FString, FString>>& SourceAndDestFiles) const;
+	virtual TMap<FString, FString> GetMappingsForRootPackageRename(const FString& SrcRoot, const FString& DstRoot, const FString& SrcBaseDir, const TArray<TPair<FString, FString>>& SourceAndDestFiles, const TMap<FString, FString>& MountPointReplacements) const;
 
 
 	virtual void GenerateAdvancedCopyDestinations(FAdvancedCopyParams& InParams, const TArray<FName>& InPackageNamesToCopy, const class UAdvancedCopyCustomization* CopyCustomization, TMap<FString, FString>& OutPackagesAndDestinations) const override;
diff --git a/Engine/Source/Developer/AssetTools/Public/IAssetTools.h b/Engine/Source/Developer/AssetTools/Public/IAssetTools.h
--- a/Engine/Source/Developer/AssetTools/Public/IAssetTools.h
+++ b/Engine/Source/Developer/AssetTools/Public/IAssetTools.h
@@ -573,7 +573,7 @@ public:
 	virtual bool PatchCopyPackageFile(const FString& SrcFile, const FString& DstFile, const TMap<FString, FString>& SearchForAndReplace) const = 0;
 
 	/** Generates the SearchAndReplace map for a PatchCopyPackageFile if all you are doing is changing the root and not the name. */
-	virtual TMap<FString, FString> GetMappingsForRootPackageRename(const FString& SrcRoot, const FString& DstRoot, const FString& SrcBaseDir, const TArray<TPair<FString, FString>>& SourceAndDestFiles) const = 0;
+	virtual TMap<FString, FString> GetMappingsForRootPackageRename(const FString& SrcRoot, const FString& DstRoot, const FString& SrcBaseDir, const TArray<TPair<FString, FString>>& SourceAndDestFiles, const TMap<FString, FString>& MountPointReplacements) const = 0;
 
 	/* Given a set of packages to copy, generate the map of those packages to destination filenames */
 	virtual void GenerateAdvancedCopyDestinations(FAdvancedCopyParams& InParams, const TArray<FName>& InPackageNamesToCopy, const UAdvancedCopyCustomization* CopyCustomization, TMap<FString, FString>& OutPackagesAndDestinations) const = 0;
