Fixes an issue in Asset Header Patching where we would find the source PackageRootName and replace any string that matches it with the destination PackageRootName. This replacing would result in potentially ignoring the mapping the user provided but also potentially rename various asset paths that happen to contain the same substring as the source PackageRoot.

We now break the passed in Long-form Package paths into their 2/3 parts, PackageRoot, PackagePath (might be empty) and PackageName, and create mappings from source -> destination only when the exact portions of the path map. Special handling is needed for the __ExternalActors__ and __ExternalObjects__ directories since remapping a package may require the subdirectories under __ExternalActors__ to need remapping as well, in addition to the SrcPackageRoot -> DstPackageRoot change.

#rnx

diff --git a/Engine/Source/Developer/AssetTools/Private/AssetTools.cpp b/Engine/Source/Developer/AssetTools/Private/AssetTools.cpp
--- a/Engine/Source/Developer/AssetTools/Private/AssetTools.cpp
+++ b/Engine/Source/Developer/AssetTools/Private/AssetTools.cpp
@@ -2217,20 +2217,73 @@ void UAssetToolsImpl::GetAllAdvancedCopySources(FName SelectedPackage, FAdvanced
 
 namespace 
 {
-bool IsSlashOrBackslash(TCHAR C) 
+bool SplitLongPackageName(FStringView LongPackageName, FStringView& PackageRoot, FStringView& PackagePath, FStringView& PackageName)
 {
-	return C == TEXT('/') || C == TEXT('\\'); 
+	if (LongPackageName.IsEmpty() || LongPackageName[0] != TEXT('/'))
+	{
+		return false;
+	}
+
+	PackageRoot = FStringView(LongPackageName.GetData() + 1); // + 1 to skip the leading '/'
+	int32 SeparatorPos;
+	if (!PackageRoot.FindChar(TEXT('/'), SeparatorPos))
+	{
+		return false;
+	}
+	PackageRoot.LeftInline(SeparatorPos);
+
+	const int32 PackagePathOffset = PackageRoot.Len() + 2; // + 2 for the leading and trailing '/'
+	if (LongPackageName.Len() < PackagePathOffset || !LongPackageName.FindLastChar(TEXT('/'), SeparatorPos))
+	{
+		return false;
+	}
+
+	// May be empty. If the PackageName is off the root there is no PackagePath
+	const int32 PackagePathLen = SeparatorPos - (PackagePathOffset - 1);
+	check(PackagePathLen >= 0);
+	PackagePath = FStringView(LongPackageName.GetData() + PackagePathOffset, PackagePathLen - !!PackagePathLen);
+
+	const int32 PackageNameOffset = PackagePathOffset + PackagePath.Len() + !PackagePath.IsEmpty();
+	PackageName = FStringView(LongPackageName.GetData() + PackageNameOffset, LongPackageName.Len() - PackageNameOffset);
+
+	return true;
 }
-	
+
 TMap<FString, FString> AllSourceAndDestPackages(const TMap<FString, FString>& SourceAndDestPackages)
 {
-	TMap<FString, FString> Result;
+	// Note, this function only correctly supports SourceAndDest mappings where all Source have the same root
+
+	// Paths under the __External root drop the package root, so create mappings we can leverage when handling
+	// those cases where the package path may have been remapped
+	TMap<FString, FString> ExternalMappings;
+	for (const TPair<FString, FString>& SrcDstPair : SourceAndDestPackages)
+	{
+		const FString& Src = SrcDstPair.Key;
+		const FString& Dst = SrcDstPair.Value;
+
+		FStringView SrcPackageRoot;
+		FStringView SrcPackagePath;
+		FStringView SrcPackageName;
+		SplitLongPackageName(Src, SrcPackageRoot, SrcPackagePath, SrcPackageName);
 
+		FStringView DstPackageRoot;
+		FStringView DstPackagePath;
+		FStringView DstPackageName;
+		SplitLongPackageName(Dst, DstPackageRoot, DstPackagePath, DstPackageName);
+
+		FStringView SrcPath = SrcPackagePath.IsEmpty() ? SrcPackageName : SrcPackagePath;
+		FStringView DstPath = DstPackagePath.IsEmpty() ? DstPackageName : DstPackagePath;
+		ExternalMappings.Add(FString(SrcPath), FString(DstPath));
+	}
+
+
+	TMap<FString, FString> Result;
 	IAssetRegistry& Registry = *IAssetRegistry::Get();
 
 	TArray< TTuple<FString, FString> > ToProcess;
 	Algo::Copy(SourceAndDestPackages, ToProcess);
 
+	TStringBuilder<NAME_SIZE> SrcDependencyBuilder;
 	while (ToProcess.Num())
 	{
 		TTuple<FString, FString> Package = ToProcess.Pop();
@@ -2244,30 +2297,106 @@ TMap<FString, FString> AllSourceAndDestPackages(const TMap<FString, FString>& So
 		Result.Add({ Package.Key, Package.Value });
 
 		TArray<FName> Dependencies;
-
 		if (!Registry.GetDependencies(FName(*Package.Key), Dependencies))
 		{
 			continue;
 		}
 
-		// Making String Views into strings because the String.Replace used inside the loop cannot use the views.
-		FString SrcPackageRoot = FString(FPackageName::SplitPackageNameRoot(Package.Key, nullptr));
-		FString DstPackageRoot = FString(FPackageName::SplitPackageNameRoot(Package.Value, nullptr));
-
+		FStringView SrcPackageRoot = FPackageName::SplitPackageNameRoot(Package.Key, nullptr);
+		FStringView DstPackageRoot = FPackageName::SplitPackageNameRoot(Package.Value, nullptr);
 		for (const FName Dependency : Dependencies)
 		{
-			const FString SrcDependencyString = Dependency.ToString();
-			
-			// checking from +1 Dependency has a leading '/' 
-			if (IsSlashOrBackslash(SrcDependencyString[0])
-				&& FStringView(*SrcDependencyString + 1, SrcPackageRoot.Len()) == SrcPackageRoot
-				&& IsSlashOrBackslash(SrcDependencyString[SrcPackageRoot.Len() + 1]))
+			Dependency.ToString(SrcDependencyBuilder);
+			FStringView SrcDependency = SrcDependencyBuilder.ToView();
+
+			if (SourceAndDestPackages.FindByHash(GetTypeHash(SrcDependency), SrcDependency))
 			{
-				// if a dep start with the package name, then we are going to copy the asset.
-				// but we need to recurse on this asset as it may have sub dependencies we don't know of yet.
-				const FString DstDependencyString = SrcDependencyString.Replace(*SrcPackageRoot, *DstPackageRoot, ESearchCase::CaseSensitive);
-				ToProcess.Add({ SrcDependencyString , DstDependencyString });
+				// We already handled this mapping
+				continue;
+			}
+
+			FStringView SrcDependencyPackageRoot;
+			FStringView SrcDependencyPackagePath;
+			FStringView SrcDependencyPackageName;
+			SplitLongPackageName(SrcDependency, SrcDependencyPackageRoot, SrcDependencyPackagePath, SrcDependencyPackageName);
+			check(!SrcDependencyPackageRoot.IsEmpty());
+
+			// Only consider dependency paths that are for the same package as our src->dst mapping
+			// If the src mapping doesn't begin with a '/' the package name will be empty, since the path isn't a package path
+			if (SrcDependencyPackageRoot != SrcPackageRoot)
+			{
+				continue;
+			}
+
+			TStringBuilder<NAME_SIZE> DstDependencyString;
+
+			// Special handling for external references. The __External[Actors__|Objects__] directory is always under the package root, may contain an
+			// arbitrary amount of subdirs but then ends with two hash subdirs. The path between the __External[Actors__|Objects__] and the two hash dirs
+			// may need remapping so we look at our external mappings to do so.
+			bool bHasExternalActorDir = SrcDependencyPackagePath.StartsWith(FPackagePath::GetExternalActorsFolderName());
+			bool bHasExternalObjectsDir = !bHasExternalActorDir && SrcDependencyPackagePath.StartsWith(FPackagePath::GetExternalObjectsFolderName());
+			if (bHasExternalActorDir || bHasExternalObjectsDir)
+			{
+				int32 RightPartStartPos;
+				if (!SrcDependencyPackagePath.FindChar(TEXT('/'), RightPartStartPos))
+				{
+					// This is a path to only the special directory, skip it no remapping is needed
+					continue;
+				}
+				RightPartStartPos++; // Skip past the '/'
+
+				// Find the start of the two hash dirs
+				// e.g. __ExternalActors__/path/of/interest/A/A9, we only want 'path/of/interest'
+				FStringView ExternalPackagePath(SrcDependencyPackagePath.GetData() + RightPartStartPos, SrcDependencyPackagePath.Len() - RightPartStartPos);
+				int32 HashDirStartPos = 0;
+				int NumHashDirsToStrip = 2;
+				while (NumHashDirsToStrip--)
+				{
+					if (ExternalPackagePath.FindLastChar(TEXT('/'), HashDirStartPos))
+					{
+						ExternalPackagePath.LeftChopInline(ExternalPackagePath.Len() - HashDirStartPos);
+					}
+				}
+
+				// Our __External[Actors|Objects]__ path is malformed
+				if (HashDirStartPos == INDEX_NONE)
+				{
+					continue;
+				}
+
+				const int32 HashPathOffset = RightPartStartPos + HashDirStartPos;
+				FStringView HashPath(SrcDependencyPackagePath.GetData() + HashPathOffset, SrcDependencyPackagePath.Len() - HashPathOffset);
+				const FString* DstExternalPackagePath = ExternalMappings.FindByHash(GetTypeHash(ExternalPackagePath), ExternalPackagePath);
+								
+				DstDependencyString.AppendChar(TEXT('/'));
+				DstDependencyString.Append(DstPackageRoot);
+				DstDependencyString.AppendChar(TEXT('/'));
+				DstDependencyString.Append(bHasExternalActorDir ? FPackagePath::GetExternalActorsFolderName() : FPackagePath::GetExternalObjectsFolderName());
+				DstDependencyString.AppendChar(TEXT('/'));
+				DstDependencyString.Append(DstExternalPackagePath ? *DstExternalPackagePath : ExternalPackagePath);
+				DstDependencyString.Append(HashPath); // HashPath already contains the leading '/'
+				DstDependencyString.AppendChar(TEXT('/'));
+				DstDependencyString.Append(SrcDependencyPackageName);
 			}
+			else
+			{ 
+				// We aren't handling a special directory so replace the package root
+				DstDependencyString.AppendChar(TEXT('/'));
+				DstDependencyString.Append(DstPackageRoot);
+				DstDependencyString.AppendChar(TEXT('/'));
+
+				if (!SrcDependencyPackagePath.IsEmpty())
+				{
+					DstDependencyString.Append(SrcDependencyPackagePath);
+					DstDependencyString.AppendChar(TEXT('/'));
+				}
+
+				DstDependencyString.Append(SrcDependencyPackageName);
+			}
+
+			// If a dep start with the package name, then we are going to copy the asset.
+			// but we need to recurse on this asset as it may have sub dependencies we don't know of yet.
+			ToProcess.Add({ FString(SrcDependency), DstDependencyString.ToString()});
 		}
 	}
 
