SourceControl: Fixed AssetRenameManager handling unresponsive source control incorrectly.

Although it did popup a dialog, it never marked the assets to rename as failed, leading calling code down the callstack to believe the rename did succeed. I've removed the dialog, as marking the assets as failed to rename will already trigger a dialog and we don't want two.

diff --git a/Engine/Source/Developer/AssetTools/Private/AssetRenameManager.cpp b/Engine/Source/Developer/AssetTools/Private/AssetRenameManager.cpp
--- a/Engine/Source/Developer/AssetTools/Private/AssetRenameManager.cpp
+++ b/Engine/Source/Developer/AssetTools/Private/AssetRenameManager.cpp
@@ -804,35 +804,50 @@ void FAssetRenameManager::PopulateAssetReferencers(TArray<FAssetRenameDataWithRe
 	}
 }
 
-bool FAssetRenameManager::UpdatePackageStatus(const TArray<FAssetRenameDataWithReferencers>& AssetsToRename) const
+bool FAssetRenameManager::UpdatePackageStatus(TArray<FAssetRenameDataWithReferencers>& AssetsToRename) const
 {
+	bool bSucceeded = true;
+
 	if (ISourceControlModule::Get().IsEnabled())
 	{
 		ISourceControlProvider& SourceControlProvider = ISourceControlModule::Get().GetProvider();
 
 		// Update the source control server availability to make sure we can do the rename operation
 		SourceControlProvider.Login();
-		if (!SourceControlProvider.IsAvailable())
-		{
-			FMessageDialog::Open(EAppMsgType::Ok, NSLOCTEXT("UnrealEd", "SourceControl_ServerUnresponsive", "Revision Control is unresponsive. Please check your connection and try again."));
-			return false;
-		}
 
-		// Gather asset package names to update SCC states in a single SCC request
-		TArray<UPackage*> PackagesToUpdate;
-		for (auto AssetIt = AssetsToRename.CreateConstIterator(); AssetIt; ++AssetIt)
+		if (SourceControlProvider.IsAvailable())
 		{
-			UObject* Asset = (*AssetIt).Asset.Get();
-			if (Asset)
+			// Gather asset package names to update SCC states in a single SCC request
+			TArray<UPackage*> PackagesToUpdate;
+			PackagesToUpdate.Reserve(AssetsToRename.Num());
+
+			for (const FAssetRenameDataWithReferencers& RenameData : AssetsToRename)
 			{
-				PackagesToUpdate.AddUnique(Asset->GetOutermost());
+				if (UObject* Asset = RenameData.Asset.Get())
+				{
+					PackagesToUpdate.AddUnique(Asset->GetOutermost());
+				}
 			}
+
+			bSucceeded = SourceControlProvider.Execute(ISourceControlOperation::Create<FUpdateStatus>(), PackagesToUpdate) == ECommandResult::Succeeded;
 		}
+		else
+		{
+			bSucceeded = false;
+		}
+	}
 
-		SourceControlProvider.Execute(ISourceControlOperation::Create<FUpdateStatus>(), PackagesToUpdate);
+	if (!bSucceeded)
+	{
+		// Mark the renames as a failure to report it later
+		for (FAssetRenameDataWithReferencers& RenameData : AssetsToRename)
+		{
+			RenameData.bRenameFailed = true;
+			RenameData.FailureReason = LOCTEXT("RenameFailedUnavailable", "Revision control is unresponsive.");
+		}
 	}
 
-	return true;
+	return bSucceeded;
 }
 
 void FAssetRenameManager::LoadReferencingPackages(TArray<FAssetRenameDataWithReferencers>& AssetsToRename, bool bLoadAllPackages, bool bCheckStatus, TArray<UPackage*>& OutReferencingPackagesToSave, TArray<UObject*>& OutSoftReferencingObjects) const
diff --git a/Engine/Source/Developer/AssetTools/Private/AssetRenameManager.h b/Engine/Source/Developer/AssetTools/Private/AssetRenameManager.h
--- a/Engine/Source/Developer/AssetTools/Private/AssetRenameManager.h
+++ b/Engine/Source/Developer/AssetTools/Private/AssetRenameManager.h
@@ -78,7 +78,7 @@ private:
 	void PopulateAssetReferencers(TArray<FAssetRenameDataWithReferencers>& AssetsToPopulate) const;
 
 	/** Updates the source control status of the packages containing the assets to rename */
-	bool UpdatePackageStatus(const TArray<FAssetRenameDataWithReferencers>& AssetsToRename) const;
+	bool UpdatePackageStatus(TArray<FAssetRenameDataWithReferencers>& AssetsToRename) const;
 
 	/**
 	 * Loads all referencing packages to assets in AssetsToRename, finds assets whose references can
