Fix AssetHeaderPatching handling of quoted paths.

- Move EndsWith check to be the first condition as it's constant time compared to check for the first quote
- Fixes the creation of the patched quoted path. The previous code was using LeftChop which incorrectly chopped off n characters from the right side of the string rather than keeping only the first n on the left. This resulted in malformed path.
- We now rebuild the string in place  to eliminate copies.

#rnx

diff --git a/Engine/Source/Developer/AssetTools/Private/AssetHeaderPatcher.cpp b/Engine/Source/Developer/AssetTools/Private/AssetHeaderPatcher.cpp
--- a/Engine/Source/Developer/AssetTools/Private/AssetHeaderPatcher.cpp
+++ b/Engine/Source/Developer/AssetTools/Private/AssetHeaderPatcher.cpp
@@ -594,13 +594,18 @@ bool FAssetHeaderPatcherInner::DoPatch(FString& InOutString)
 	{
 		// Patch quoted paths.
 		// Path occurs to the right of the first "'" 
-		int Idx{};
-		if (InOutString.FindChar(TCHAR('\''), Idx) && InOutString.EndsWith(TEXT("'")))
+		int FirstQuotePos = INDEX_NONE;
+		if (InOutString.EndsWith(TEXT("'")) && InOutString.FindChar(TCHAR('\''), FirstQuotePos))
 		{
-			FStringView MaybeReplacement = Find(SearchAndReplace, FStringView(*InOutString + Idx + 1, InOutString.Len() - (Idx + 2)));
+			// +1 to skip first quote, +2 to account for the length minus the first and last quotes
+			FStringView StrippedQuotedPath = FStringView(*InOutString + FirstQuotePos + 1, InOutString.Len() - (FirstQuotePos + 2));
+			FStringView MaybeReplacement = Find(SearchAndReplace, StrippedQuotedPath);
 			if (!MaybeReplacement.IsEmpty())
 			{
-				InOutString = InOutString.LeftChop(Idx + 1) + MaybeReplacement + TCHAR('\'');
+				InOutString.Reserve(FirstQuotePos + MaybeReplacement.Len() + 2); // + 2 for the two quotes
+				InOutString.LeftInline(FirstQuotePos + 1, EAllowShrinking::No); // Keep the original prefix before the quotes + the first quote
+				InOutString.Append(MaybeReplacement);
+				InOutString.AppendChar(TCHAR('\''));
 				return true;
 			}
 		}
