Adds support for multiple plugins in AssetHeaderPatcher.

We now separate the passed in remapping paths to be categorized by plugin when considering mappings for external actors/objects.

#rnx

diff --git a/Engine/Source/Developer/AssetTools/Private/AssetTools.cpp b/Engine/Source/Developer/AssetTools/Private/AssetTools.cpp
--- a/Engine/Source/Developer/AssetTools/Private/AssetTools.cpp
+++ b/Engine/Source/Developer/AssetTools/Private/AssetTools.cpp
@@ -2251,11 +2251,9 @@ bool SplitLongPackageName(FStringView LongPackageName, FStringView& PackageRoot,
 
 TMap<FString, FString> AllSourceAndDestPackages(const TMap<FString, FString>& SourceAndDestPackages)
 {
-	// Note, this function only correctly supports SourceAndDest mappings where all Source have the same root
-
-	// Paths under the __External root drop the package root, so create mappings we can leverage when handling
-	// those cases where the package path may have been remapped
-	TMap<FString, FString> ExternalMappings;
+	// Paths under the __External root drop the package root, so create mappings, per plugin, 
+	// we can leverage when handling those cases where the package path may have been remapped
+	TMap<FString, TMap<FString, FString>> PluginExternalMappings;
 	for (const TPair<FString, FString>& SrcDstPair : SourceAndDestPackages)
 	{
 		const FString& Src = SrcDstPair.Key;
@@ -2271,12 +2269,12 @@ TMap<FString, FString> AllSourceAndDestPackages(const TMap<FString, FString>& So
 		FStringView DstPackageName;
 		SplitLongPackageName(Dst, DstPackageRoot, DstPackagePath, DstPackageName);
 
+		TMap<FString, FString>& ExternalMappings = PluginExternalMappings.FindOrAddByHash(GetTypeHash(SrcPackageRoot), FString(SrcPackageRoot));
 		FStringView SrcPath = SrcPackagePath.IsEmpty() ? SrcPackageName : SrcPackagePath;
 		FStringView DstPath = DstPackagePath.IsEmpty() ? DstPackageName : DstPackagePath;
 		ExternalMappings.Add(FString(SrcPath), FString(DstPath));
 	}
 
-
 	TMap<FString, FString> Result;
 	IAssetRegistry& Registry = *IAssetRegistry::Get();
 
@@ -2366,7 +2364,13 @@ TMap<FString, FString> AllSourceAndDestPackages(const TMap<FString, FString>& So
 
 				const int32 HashPathOffset = RightPartStartPos + HashDirStartPos;
 				FStringView HashPath(SrcDependencyPackagePath.GetData() + HashPathOffset, SrcDependencyPackagePath.Len() - HashPathOffset);
-				const FString* DstExternalPackagePath = ExternalMappings.FindByHash(GetTypeHash(ExternalPackagePath), ExternalPackagePath);
+				const TMap<FString, FString>* ExternalMappings = PluginExternalMappings.FindByHash(GetTypeHash(SrcPackageRoot), SrcPackageRoot);
+				if (!ExternalMappings)
+				{
+					// We have no mapping for this dependency's external actors/objects
+					continue;
+				}
+				const FString* DstExternalPackagePath = ExternalMappings->FindByHash(GetTypeHash(ExternalPackagePath), ExternalPackagePath);
 								
 				DstDependencyString.AppendChar(TEXT('/'));
 				DstDependencyString.Append(DstPackageRoot);
