Adds GatherableTextData support to AssetHeaderPatching

- Now look for the optional GatherableTextData section when deserializing, patch the contents and write back out the data, patching up the summary size as necessary since this section does not have a fixed size based on item counts (i.e. it doesn't use the name table, so a string changes affect the size)
- Added additional mapping for VerseAssetPaths which appear in GeneratedTextData (e.g. /localhost/Normal/PackagePath/PackagePath)

#rnx

diff --git a/Engine/Source/Developer/AssetTools/Private/AssetHeaderPatcher.cpp b/Engine/Source/Developer/AssetTools/Private/AssetHeaderPatcher.cpp
--- a/Engine/Source/Developer/AssetTools/Private/AssetHeaderPatcher.cpp
+++ b/Engine/Source/Developer/AssetTools/Private/AssetHeaderPatcher.cpp
@@ -11,6 +11,7 @@
 #include "AssetRegistry/IAssetRegistry.h"
 #include "AssetRegistry/AssetData.h"
 
+#include "Internationalization/GatherableTextData.h"
 #include "Misc/EnumerateRange.h"
 #include "Misc/FileHelper.h"
 #include "Misc/PackageName.h"
@@ -68,6 +69,7 @@ namespace
 		Summary,
 		NameTable,
 		SoftPathTable,
+		GatherableTextDataTable,
 		ImportTable,
 		ExportTable,
 		SoftPackageReferencesTable,
@@ -87,7 +89,7 @@ namespace
 	{
 		NameTable,
 		SoftObjectPathList,
-		GatherableTextDataOffset,
+		GatherableTextDataTable,
 		ImportTable,
 		ExportTable,
 		DependsTable,
@@ -267,6 +269,7 @@ public:
 	bool DoPatch(FName& InOutName);
 	bool DoPatch(FSoftObjectPath& InOutSoft);
 	bool DoPatch(FTopLevelAssetPath& InOutPath);
+	bool DoPatch(FGatherableTextData& InGatherablerTextData);
 
 	FAssetHeaderPatcher::EResult PatchHeader();
 	FAssetHeaderPatcher::EResult PatchHeader_Deserialize();
@@ -287,9 +290,10 @@ public:
 		int64 SummarySize = -1;
 		int64 NameTableSize = -1;
 		int64 SoftObjectPathListSize = -1;
-		int64 SoftPackageReferencesListSize = -1;
+		int64 GatherableTextDataSize = -1;
 		int64 ImportTableSize = -1;
 		int64 ExportTableSize = -1;
+		int64 SoftPackageReferencesListSize = -1;
 		int64 ThumbnailTableSize = -1;
 		int64 AssetRegistryDataSize = -1;
 		int64 PackageTrailerSize = -1;
@@ -305,11 +309,12 @@ public:
 	FHeaderInformation HeaderInformation;
 	FPackageFileSummary Summary;
 	TArray<FName> NameTable;
-	TArray<FSoftObjectPath> SoftObjectPathTable;
-	TArray<FName> SoftPackageReferencesTable;
 	TMap<FNameEntryId, int32> NameToIndexMap;
+	TArray<FSoftObjectPath> SoftObjectPathTable;
+	TArray<FGatherableTextData> GatherableTextDataTable;
 	TArray<FObjectImport> ImportTable;
 	TArray<FObjectExport> ExportTable;
+	TArray<FName> SoftPackageReferencesTable;
 	TArray<FThumbnailEntry> ThumbnailTable;
 
 	// Asset registry data information
@@ -431,6 +436,23 @@ FAssetHeaderPatcher::EResult FAssetHeaderPatcherInner::PatchHeader_Deserialize()
 		HeaderInformation.SoftObjectPathListSize = 0;
 	}
 
+	if (Summary.GatherableTextDataCount > 0)
+	{
+		MemAr.Seek(Summary.GatherableTextDataOffset);
+		GatherableTextDataTable.Reserve(Summary.GatherableTextDataCount);
+		for (int32 GatherableTextDataIndex = 0; GatherableTextDataIndex < Summary.GatherableTextDataCount; ++GatherableTextDataIndex)
+		{
+			FGatherableTextData& GatherableTextData = GatherableTextDataTable.Emplace_GetRef();
+			MemAr << GatherableTextData;
+		}
+
+		HeaderInformation.GatherableTextDataSize = MemAr.Tell() - Summary.GatherableTextDataOffset;
+	}
+	else
+	{
+		HeaderInformation.GatherableTextDataSize = 0;
+	}
+
 #define UE_CHECK_AND_SET_ERROR_AND_RETURN(EXP)	\
 	do											\
 	{											\
@@ -707,12 +729,93 @@ bool FAssetHeaderPatcherInner::DoPatch(FTopLevelAssetPath& InOutPath)
 	return bPatching;
 }
 
+bool FAssetHeaderPatcherInner::DoPatch(FGatherableTextData& InGatherableTextData)
+{
+	// There are various fields in FGatherableTextData however only one pertains to 
+	// asset paths and types, SourceSiteContexts.SiteDescription. The rest are contextual
+	// key-value pairs of text which are not references to assets/types and thus do not need patching
+	// (at least we can't understand the context a priori to know if specialized code
+	// may try to load from these strings)
+
+	bool bDidPatch = false;
+	for (FTextSourceSiteContext& SourceSiteContext : InGatherableTextData.SourceSiteContexts)
+	{
+		FStringView ClassName;
+		FStringView PackagePath;
+		FStringView ObjectName;
+		FStringView SubObjectName;
+		FPackageName::SplitFullObjectPath(SourceSiteContext.SiteDescription, ClassName, PackagePath, ObjectName, SubObjectName, true /*bDetectClassName*/);
+
+		// If we don't have a PackagePath, do nothing. Replacing object names alone without the package path is likely going to lead to incorrect replacements
+		if (PackagePath.IsEmpty())
+		{
+			continue;
+		}
+
+		const FString* DstPackagePath = SearchAndReplace.FindByHash(GetTypeHash(PackagePath), PackagePath);
+		if (!DstPackagePath)
+		{
+			continue;
+		}
+
+		// Note, we don't provide remappings down to the sub-object level so we can skip checking for and handling of sub-object specific renames
+		// We also don't bother stripping off the ClassName below, nor do we restore a prepended classname in the Dst string by scanning the replacement string
+		// if one existed in the Src string. I don't believe SiteDescriptions will ever contain a ClassName, but just in case I'm wrong we'll add a check.
+		check(ClassName.IsEmpty());
+
+		bDidPatch = true;
+
+		if (ObjectName.IsEmpty())
+		{
+			SourceSiteContext.SiteDescription = *DstPackagePath;
+			continue;
+		}
+
+		// If we have a object the whole path can be used (minus SubObjectName )
+		FStringView FullObjectPath(SourceSiteContext.SiteDescription);
+		FullObjectPath.LeftChopInline(SubObjectName.Len() + !!SubObjectName.Len()); // Strip off the SubObjectName + 1 for the ':' separator (if we have a subobject name at all)
+
+		if (const FString* DstFullObjectPath = SearchAndReplace.FindByHash(GetTypeHash(FullObjectPath), FullObjectPath))
+		{
+			// We have a replacement for the PackageName.ObjectName but not the subobject, so append the subobject if necessary
+			if (!SubObjectName.IsEmpty())
+			{
+				// We are writing over the original string, so capture our stringview before proceeding
+				FString LocalSubObjectName(SubObjectName);
+				SourceSiteContext.SiteDescription.Reserve(DstFullObjectPath->Len() + SubObjectName.Len() + 1); // + 1 for ':'
+				SourceSiteContext.SiteDescription = *DstFullObjectPath;
+				SourceSiteContext.SiteDescription.AppendChar(TEXT(':'));
+				SourceSiteContext.SiteDescription.Append(LocalSubObjectName);
+			}
+			else
+			{
+				SourceSiteContext.SiteDescription = *DstFullObjectPath;
+			}
+		}
+		else
+		{
+			// We only have a PackagePath replacement so only replace that portion of the string and keep the object name and sub-object paths untouched.
+			// Take all text from the start of the ObjectName until the end of SiteDescription
+			// This will let us include the subobject path if any since we know SiteDescription is null-terminated
+			
+			// We are writing over the original string, so capture our stringview before proceeding
+			FString SrcFullObjectPath(ObjectName.GetData());
+			SourceSiteContext.SiteDescription.Reserve(DstPackagePath->Len() + SrcFullObjectPath.Len() + 1); // + 1 for '.'
+			SourceSiteContext.SiteDescription = *DstPackagePath;
+			SourceSiteContext.SiteDescription.AppendChar(TEXT('.'));
+			SourceSiteContext.SiteDescription.Append(SrcFullObjectPath);
+		}
+	}
+	
+	return bDidPatch;
+}
+
 void FAssetHeaderPatcherInner::PatchHeader_PatchSections()
 {
 	DoPatch(Summary.PackageName);
 
 	{	// Patch the Name table
-		// This data is used by all FGNames in the file.
+		// This data is used by all FNames in the file.
 		// So if we patch a Identifier we append it to the name table.
 		// This is because there may be FNames in data we dont want to patch in structures we dont look at.
 		// If we patch the ident at inplace, then we would change those names.
@@ -754,6 +857,12 @@ void FAssetHeaderPatcherInner::PatchHeader_PatchSections()
 		DoPatch(PathRef);
 	}
 
+	// GatherableTextData table
+	for (FGatherableTextData& GatherableTextData : GatherableTextDataTable)
+	{
+		DoPatch(GatherableTextData);
+	}
+
 	// Import table
 	for (FObjectImport& Import : ImportTable)
 	{
@@ -847,16 +956,17 @@ FAssetHeaderPatcher::EResult FAssetHeaderPatcherInner::PatchHeader_WriteDestinat
 {
 	// Serialize modified sections and reconstruct the file	
 	// Original offsets and sizes of any sections that will be patched
-	//	Tag												Offset									Size								bRequired
+	//	  Tag											Offset									Size												bRequired
 	const FSectionData SourceSections[] = {
-		{ EPatchedSection::Summary,						0,										HeaderInformation.SummarySize,			true },
-		{ EPatchedSection::NameTable,					Summary.NameOffset,						HeaderInformation.NameTableSize,		true },
-		{ EPatchedSection::SoftPathTable,				Summary.SoftObjectPathsOffset,			HeaderInformation.SoftObjectPathListSize, false },
-		{ EPatchedSection::ImportTable,					Summary.ImportOffset,					HeaderInformation.ImportTableSize,		true },
-		{ EPatchedSection::ExportTable,					Summary.ExportOffset,					HeaderInformation.ExportTableSize,		true },
-		{ EPatchedSection::SoftPackageReferencesTable,	Summary.SoftPackageReferencesOffset,	HeaderInformation.SoftPackageReferencesListSize, false },
-		{ EPatchedSection::ThumbnailTable,				Summary.ThumbnailTableOffset,			HeaderInformation.ThumbnailTableSize,	false },
-		{ EPatchedSection::AssetRegistryData,			Summary.AssetRegistryDataOffset,		AssetRegistryData.SectionSize,			true },
+		{ EPatchedSection::Summary,						0,										HeaderInformation.SummarySize,						true	},
+		{ EPatchedSection::NameTable,					Summary.NameOffset,						HeaderInformation.NameTableSize,					true	},
+		{ EPatchedSection::SoftPathTable,				Summary.SoftObjectPathsOffset,			HeaderInformation.SoftObjectPathListSize,			false	},
+		{ EPatchedSection::GatherableTextDataTable,		Summary.GatherableTextDataOffset,		HeaderInformation.GatherableTextDataSize,			false	},
+		{ EPatchedSection::ImportTable,					Summary.ImportOffset,					HeaderInformation.ImportTableSize,					true	},
+		{ EPatchedSection::ExportTable,					Summary.ExportOffset,					HeaderInformation.ExportTableSize,					true	},
+		{ EPatchedSection::SoftPackageReferencesTable,	Summary.SoftPackageReferencesOffset,	HeaderInformation.SoftPackageReferencesListSize,	false	},
+		{ EPatchedSection::ThumbnailTable,				Summary.ThumbnailTableOffset,			HeaderInformation.ThumbnailTableSize,				false	},
+		{ EPatchedSection::AssetRegistryData,			Summary.AssetRegistryDataOffset,		AssetRegistryData.SectionSize,						true	},
 	};
 
 	const int32 SourceTotalHeaderSize = Summary.TotalHeaderSize;
@@ -918,7 +1028,6 @@ FAssetHeaderPatcher::EResult FAssetHeaderPatcherInner::PatchHeader_WriteDestinat
 		const FSectionData& SourceSection = SourceSections[SectionIdx];
 
 		// skip processing empty non required chunks.
-		// really the only option section is the Thumbnails, and its annoying its ion the middle.
 		if (!SourceSection.bRequired && SourceSection.Size <= 0)
 		{
 			continue;
@@ -986,6 +1095,23 @@ FAssetHeaderPatcher::EResult FAssetHeaderPatcherInner::PatchHeader_WriteDestinat
 			break;
 		}
 
+		case EPatchedSection::GatherableTextDataTable:
+		{
+			const int64 GatherableTableStartOffset = Writer.Tell();
+			for (FGatherableTextData& GatherableTextData : GatherableTextDataTable)
+			{
+				Writer << GatherableTextData;
+			}
+			const int64 TableSize = Writer.Tell() - GatherableTableStartOffset;
+			const int64 Delta = TableSize - SourceSection.Size;
+			PatchSummaryOffsets(Summary, GatherableTableStartOffset, Delta);
+			Summary.TotalHeaderSize += (int32)Delta;
+			check(Summary.GatherableTextDataCount == GatherableTextDataTable.Num());
+			check(Summary.GatherableTextDataOffset == GatherableTableStartOffset);
+
+			break;
+		}
+
 		case EPatchedSection::ImportTable:
 		{
 			const int64 ImportTableStartOffset = Writer.Tell();
diff --git a/Engine/Source/Developer/AssetTools/Private/AssetTools.cpp b/Engine/Source/Developer/AssetTools/Private/AssetTools.cpp
--- a/Engine/Source/Developer/AssetTools/Private/AssetTools.cpp
+++ b/Engine/Source/Developer/AssetTools/Private/AssetTools.cpp
@@ -2503,6 +2503,28 @@ TMap<FString, FString> GenerateAdditionalAssetMappings(const TMap<FString, FStri
 			}
 		}
 
+		// VerseAssetPath mapping used in GatherableTextData
+		// e.g. /localhost/Some/Package/Path/PackageName/PackageObject (note the use of '/' instead of '.' for the top-level Package object)
+		{
+			// Using a constant to avoid pulling in, what would otherwise be, unnecessary dependencies
+			const TCHAR* VerseRoot = TEXT("/localhost");
+
+			NameBuilder.Reset();
+			NameBuilder.Append(VerseRoot);
+			NameBuilder.Append(SrcNameString);
+			NameBuilder.AppendChar(TEXT('/'));
+			NameBuilder.Append(SrcPackageName);
+			FString SrcObjectName = NameBuilder.ToString();
+
+			NameBuilder.Reset();
+			NameBuilder.Append(VerseRoot);
+			NameBuilder.Append(DstNameString);
+			NameBuilder.AppendChar(TEXT('/'));
+			NameBuilder.Append(DstPackageName);
+			FString DstObjectName = NameBuilder.ToString();
+			Result.Add(MoveTemp(SrcObjectName), MoveTemp(DstObjectName));
+		}
+		
 		if (SrcPackageName != DstPackageName)
 		{
 			// PackageName (without PackageRoot or PackagePath) mapping (numbered suffix included)
