Submitting on behalf of andrew.phillips: Fix for Header patching copy of Grey Box level.

Changes NameTable patching to use full DisplayName for FNames when patching to preserve any number that may need patching. FNames are then given a sentinel number value to ensure the number is not stripped when inspecting the PlainStringName of the FName. NameTable entries don't maintain numbers for FNames so the sentinel value of '1' will be ignored come serialization time.

#rnx

diff --git a/Engine/Source/Developer/AssetTools/Private/AssetHeaderPatcher.cpp b/Engine/Source/Developer/AssetTools/Private/AssetHeaderPatcher.cpp
--- a/Engine/Source/Developer/AssetTools/Private/AssetHeaderPatcher.cpp
+++ b/Engine/Source/Developer/AssetTools/Private/AssetHeaderPatcher.cpp
@@ -822,11 +822,15 @@ void FAssetHeaderPatcherInner::PatchHeader_PatchSections()
 		TArray<FName> ToAppend;
 		for (FName& Name : NameTable)
 		{
-			FName TmpName = Name;
+			FString TmpNameStr = Name.ToString(); // will contain a "[0-9]+" suffix, if any
 
-			if (DoPatch(TmpName))
+			if (DoPatch(TmpNameStr))
 			{
-				FString TmpNameStr = Name.GetPlainNameString();
+				// Construct a new FName based on our patched string.
+				// We give it "1" as its number so that if the string has a "[0-9]+" suffix we keep it as part of the plain string name. 
+				// When we write the NameTable back out, it will serialize the string with the number, but not the "1" set here.
+				FName TmpName = FName(*TmpNameStr, 1); 
+
 				// If the string does not contain any path separators, then call it an Identifier.
 				if (!(TmpNameStr.Contains(TEXT("/")) || TmpNameStr.Contains(TEXT("\\"))))
 				{
