Redo for fix for Header patching of greybox level.
Addition fix for when copying to a name_[0-9]+ project.
Additional better logging should the FName writter fail. (now unlinky due to additional FName table entries being genorated on demand.)

diff --git a/Engine/Source/Developer/AssetTools/Private/AssetHeaderPatcher.cpp b/Engine/Source/Developer/AssetTools/Private/AssetHeaderPatcher.cpp
--- a/Engine/Source/Developer/AssetTools/Private/AssetHeaderPatcher.cpp
+++ b/Engine/Source/Developer/AssetTools/Private/AssetHeaderPatcher.cpp
@@ -45,7 +45,7 @@ namespace
 
 			if (MaybeIndex == nullptr)
 			{
-				UE_LOG(LogAssetHeaderPatcher, Error, TEXT("Cannot serialize FName %s because it is not in the name table for %s"), *Name.ToString(), *GetArchiveName());
+				ErrorMessage += FString::Printf(TEXT("Cannot serialize FName %s because it is not in the name table for %s\n"), *Name.ToString(), *GetArchiveName());
 				SetCriticalError();
 				return *this;
 			}
@@ -60,8 +60,14 @@ namespace
 			return *this;
 		}
 
+		const FString& GetErrorMessage() const
+		{
+			return ErrorMessage;
+		}
+
 	private:
 		const TMap<FNameEntryId, int32>& NameToIndexMap;
+		FString ErrorMessage;
 	};
 
 	enum class EPatchedSection
@@ -700,6 +706,18 @@ bool FAssetHeaderPatcherInner::DoPatch(FName& InOutName)
 			// Value did not have a tail, so use the original number.
 			InOutName = FName(Value, InOutName.GetNumber());
 		}
+
+		// For specific cases we want the number to be parsed in the FName.
+		// This is usually when we use it as a path.
+		// It may be, that dont have a FName in the table for that case at this point, so we add it.
+		// This will probably contribute to unused FNames in the FName table, 
+		// but these will most likely be removed on a subsequent full save.
+		if (!NameToIndexMap.Contains(InOutName.GetDisplayIndex()))
+		{
+			int32 Indx = NameTable.Num();
+			NameTable.Add(InOutName);
+			NameToIndexMap.Add(InOutName.GetDisplayIndex(), Indx);
+		}
 		return true;
 	}
 	return false;
@@ -822,11 +840,16 @@ void FAssetHeaderPatcherInner::PatchHeader_PatchSections()
 		TArray<FName> ToAppend;
 		for (FName& Name : NameTable)
 		{
-			FName TmpName = Name;
+			FNameEntryId OriginalFNameEntryId = Name.GetDisplayIndex();
+			FString TmpNameStr = Name.ToString(); // will contain a "[0-9]+" suffix, if any
 
-			if (DoPatch(TmpName))
+			if (DoPatch(TmpNameStr))
 			{
-				FString TmpNameStr = Name.GetPlainNameString();
+				// Construct a new FName based on our patched string.
+				// If the string has a "[0-9]+" suffix we keep it as part of the plain string name.
+				// When we write the NameTable back out, it will serialize the string with the number, but not the number set here.
+				FName TmpName = FName(*TmpNameStr, NAME_NO_NUMBER);
+
 				// If the string does not contain any path separators, then call it an Identifier.
 				if (!(TmpNameStr.Contains(TEXT("/")) || TmpNameStr.Contains(TEXT("\\"))))
 				{
@@ -950,6 +973,9 @@ void FAssetHeaderPatcherInner::PatchHeader_PatchSections()
 			}
 		}
 	}
+
+	// Just incase new names were added due to number parsing
+	Summary.NameCount = NameTable.Num();
 }
 
 FAssetHeaderPatcher::EResult FAssetHeaderPatcherInner::PatchHeader_WriteDestinationFile()
@@ -1069,6 +1095,8 @@ FAssetHeaderPatcher::EResult FAssetHeaderPatcherInner::PatchHeader_WriteDestinat
 				check(Entry);
 				Entry->Write(Writer);
 			}
+			checkf(!Writer.IsCriticalError(), TEXT("Issue writing %s"), *Writer.GetErrorMessage());
+
 			const int64 NameTableSize = Writer.Tell() - NameTableStartOffset;
 			const int64 Delta = NameTableSize - SourceSection.Size;
 			PatchSummaryOffsets(Summary, NameTableStartOffset, Delta);
@@ -1086,6 +1114,8 @@ FAssetHeaderPatcher::EResult FAssetHeaderPatcherInner::PatchHeader_WriteDestinat
 			{
 				PathRef.SerializePath(Writer);
 			}
+			checkf(!Writer.IsCriticalError(), TEXT("Issue writing %s"), *Writer.GetErrorMessage());
+
 			const int64 TableSize = Writer.Tell() - TableStartOffset;
 			const int64 Delta = TableSize - SourceSection.Size;
 			checkf(Delta == 0, TEXT("Delta should be Zero. is %d"), (int)Delta);
@@ -1102,6 +1132,8 @@ FAssetHeaderPatcher::EResult FAssetHeaderPatcherInner::PatchHeader_WriteDestinat
 			{
 				Writer << GatherableTextData;
 			}
+			checkf(!Writer.IsCriticalError(), TEXT("Issue writing %s"), *Writer.GetErrorMessage());
+
 			const int64 TableSize = Writer.Tell() - GatherableTableStartOffset;
 			const int64 Delta = TableSize - SourceSection.Size;
 			PatchSummaryOffsets(Summary, GatherableTableStartOffset, Delta);
@@ -1119,6 +1151,8 @@ FAssetHeaderPatcher::EResult FAssetHeaderPatcherInner::PatchHeader_WriteDestinat
 			{
 				Writer << Import;
 			}
+			checkf(!Writer.IsCriticalError(), TEXT("Issue writing %s"), *Writer.GetErrorMessage());
+
 			const int64 ImportTableSize = Writer.Tell() - ImportTableStartOffset;
 			const int64 Delta = ImportTableSize - SourceSection.Size;
 			check(Delta == 0);
@@ -1138,6 +1172,8 @@ FAssetHeaderPatcher::EResult FAssetHeaderPatcherInner::PatchHeader_WriteDestinat
 			{
 				Writer << Export;
 			}
+			checkf(!Writer.IsCriticalError(), TEXT("Issue writing %s"), *Writer.GetErrorMessage());
+
 			const int64 ExportTableSize = Writer.Tell() - ExportTableStartOffset;
 			const int64 Delta = ExportTableSize - SourceSection.Size;
 			check(Delta == 0);
@@ -1155,6 +1191,8 @@ FAssetHeaderPatcher::EResult FAssetHeaderPatcherInner::PatchHeader_WriteDestinat
 			{
 				Writer << Reference;
 			}
+			checkf(!Writer.IsCriticalError(), TEXT("Issue writing %s"), *Writer.GetErrorMessage());
+
 			const int64 TableSize = Writer.Tell() - TableStartOffset;
 			const int64 Delta = TableSize - SourceSection.Size;
 			checkf(Delta == 0, TEXT("Delta should be Zero. is %d"), (int)Delta);
@@ -1178,6 +1216,8 @@ FAssetHeaderPatcher::EResult FAssetHeaderPatcherInner::PatchHeader_WriteDestinat
 				Entry.FileOffset += (int32)ThumbnailTableDeltaOffset; // Thumbnail payloads immediately follow the table, so we can just apply the section offset delta here
 				Writer << Entry.FileOffset;
 			}
+			checkf(!Writer.IsCriticalError(), TEXT("Issue writing %s"), *Writer.GetErrorMessage());
+
 			const int64 ThumbnailTableSize = Writer.Tell() - ThumbnailTableStartOffset;
 			checkf(ThumbnailTableStartOffset == Summary.ThumbnailTableOffset, TEXT("%zd == %zd"), ThumbnailTableStartOffset, Summary.ThumbnailTableOffset);
 			check(ThumbnailTableSize == SourceSection.Size); // We only patch thumbnail table offsets, we should not be patching size
@@ -1211,6 +1251,7 @@ FAssetHeaderPatcher::EResult FAssetHeaderPatcherInner::PatchHeader_WriteDestinat
 					Writer << TagData.Value;
 				}
 			}
+			checkf(!Writer.IsCriticalError(), TEXT("Issue writing %s"), *Writer.GetErrorMessage());
 
 			const int64 AssetRegistryDataSize = Writer.Tell() - AssetRegistryDataStartOffset;
 			const int64 Delta = AssetRegistryDataSize - SourceSection.Size;
