- Change Cascade asset creation name from Legacy to Deprecated

#rnx
#tests Created Cascade asset, observed change in text

diff --git a/Engine/Source/Developer/AssetTools/Private/AssetTypeActions/AssetTypeActions_ParticleSystem.h b/Engine/Source/Developer/AssetTools/Private/AssetTypeActions/AssetTypeActions_ParticleSystem.h
--- a/Engine/Source/Developer/AssetTools/Private/AssetTypeActions/AssetTypeActions_ParticleSystem.h
+++ b/Engine/Source/Developer/AssetTools/Private/AssetTypeActions/AssetTypeActions_ParticleSystem.h
@@ -12,7 +12,7 @@ class FAssetTypeActions_ParticleSystem : public FAssetTypeActions_Base
 {
 public:
 	// IAssetTypeActions Implementation
-	virtual FText GetName() const override { return NSLOCTEXT("AssetTypeActions", "AssetTypeActions_ParticleSystem", "Cascade Particle System (Legacy)"); }
+	virtual FText GetName() const override { return NSLOCTEXT("AssetTypeActions", "AssetTypeActions_ParticleSystem", "Cascade Particle System (Deprecated)"); }
 	virtual FColor GetTypeColor() const override { return FColor(255,255,255); }
 	virtual UClass* GetSupportedClass() const override { return UParticleSystem::StaticClass(); }
 	virtual void GetActions(const TArray<UObject*>& InObjects, struct FToolMenuSection& Section) override;
