- Expose IAssetTools::GetAdditionalPatchCopyMappings and added a exclusion check for external actor/object assets
- Rename IAssetTools::GetMappingsForRootPackageRename to GetPatchCopyMappingsForRootRename for clarity
#rnx

diff --git a/Engine/Source/Developer/AssetTools/Private/AssetTools.cpp b/Engine/Source/Developer/AssetTools/Private/AssetTools.cpp
--- a/Engine/Source/Developer/AssetTools/Private/AssetTools.cpp
+++ b/Engine/Source/Developer/AssetTools/Private/AssetTools.cpp
@@ -2442,157 +2442,6 @@
 	// If we get here, the whole Name is a number
 	return INDEX_NONE;
 }
-
-TMap<FString, FString> GenerateAdditionalAssetMappings(const TMap<FString, FString>& SourceAndDestPackages)
-{
-	TMap<FString, FString> Result;
-
-	TStringBuilder<NAME_SIZE> NameBuilder;
-	for (const TTuple<FString, FString>& Package : SourceAndDestPackages)
-	{
-		const FString& SrcNameString = Package.Key;
-		const FString& DstNameString = Package.Value;
-
-		// FPathViews::GetBaseFilename gives the same result as FPackageName::GetShortName
-		// for a file path, but returns a StringView not a String.
-		FStringView SrcPackageName = FPathViews::GetBaseFilename(SrcNameString);
-		FStringView DstPackageName = FPathViews::GetBaseFilename(DstNameString);
-
-		// Path.ObjectName mapping
-		{
-			NameBuilder.Reset();
-			NameBuilder.Append(SrcNameString);
-			NameBuilder.AppendChar(TEXT('.'));
-			NameBuilder.Append(SrcPackageName);
-			FString SrcObjectName = NameBuilder.ToString();
-
-			NameBuilder.Reset();
-			NameBuilder.Append(DstNameString);
-			NameBuilder.AppendChar(TEXT('.'));
-			NameBuilder.Append(DstPackageName);
-			FString DstObjectName = NameBuilder.ToString();
-			Result.Add(MoveTemp(SrcObjectName), MoveTemp(DstObjectName));
-		}
-
-		// For package names with a '_[0-9]+' suffix: Add a mapping with the suffix stripped but 
-		// maintain the full long path provided.
-		// NameTable entries have this suffix removed and we won't match them otherwise.
-		{
-			int32 SrcNameStrPlainStringLen = EndsWithStrippableNumber(SrcNameString);				// PackageRoot_0/PackagePath_0/PackageName_0
-			if (SrcNameStrPlainStringLen != INDEX_NONE)
-			{
-				FStringView SrcNamePlainString(*SrcNameString, SrcNameStrPlainStringLen);			// PackageRoot_0/PackagePath_0/PackageName
-
-				// It's valid to rename a numbered name to a non-numbered name, 
-				// default to the original name and use the non-numbered name if found later
-				FStringView DstNamePlainString = DstNameString;
-				int32 DstNameStrPlainStringLen = EndsWithStrippableNumber(DstNameString);
-				if (DstNameStrPlainStringLen != INDEX_NONE)
-				{
-					DstNamePlainString.LeftInline(DstNameStrPlainStringLen);
-				}
-				Result.Add(FString(SrcNamePlainString), FString(DstNamePlainString));				// PackageRoot_0/PackagePath_0/PackageName => DstPackageRoot_0/DstPackagePath_0/DstPackageName
-
-				FStringView SrcShortPackageName = FPathViews::GetBaseFilename(SrcNamePlainString);	// PackageName
-				NameBuilder.Reset();
-				NameBuilder.Append(SrcNameString);
-				NameBuilder.AppendChar(TEXT('.'));
-				NameBuilder.Append(SrcShortPackageName);
-				FString SrcPathWithPlainObjectName = NameBuilder.ToString();
-
-				FStringView DstShortPackageName = FPathViews::GetBaseFilename(DstNamePlainString);	// PackageName
-				NameBuilder.Reset();
-				NameBuilder.Append(DstNameString);
-				NameBuilder.AppendChar(TEXT('.'));
-				NameBuilder.Append(DstShortPackageName);
-				FString DstPathWithObjectName = NameBuilder.ToString();
-				Result.Add(MoveTemp(SrcPathWithPlainObjectName), MoveTemp(DstPathWithObjectName));	// PackageRoot_0/PackagePath_0/PackageName_0.PackageName => DstPackageRoot_0/DstPackagePath_0/DstPackageName_0.DstPackageName
-			}
-		}
-
-		// VerseAssetPath mapping used in GatherableTextData
-		// e.g. /localhost/Some/Package/Path/PackageName/PackageObject (note the use of '/' instead of '.' for the top-level Package object)
-		{
-			// Using a constant to avoid pulling in, what would otherwise be, unnecessary dependencies
-			const TCHAR* VerseRoot = TEXT("/localhost");
-
-			NameBuilder.Reset();
-			NameBuilder.Append(VerseRoot);
-			NameBuilder.Append(SrcNameString);
-			NameBuilder.AppendChar(TEXT('/'));
-			NameBuilder.Append(SrcPackageName);
-			FString SrcObjectName = NameBuilder.ToString();
-
-			NameBuilder.Reset();
-			NameBuilder.Append(VerseRoot);
-			NameBuilder.Append(DstNameString);
-			NameBuilder.AppendChar(TEXT('/'));
-			NameBuilder.Append(DstPackageName);
-			FString DstObjectName = NameBuilder.ToString();
-			Result.Add(MoveTemp(SrcObjectName), MoveTemp(DstObjectName));
-		}
-		
-		if (SrcPackageName != DstPackageName)
-		{
-			// PackageName (without PackageRoot or PackagePath) mapping (numbered suffix included)
-			{
-				Result.Add({ FString(SrcPackageName), FString(DstPackageName) });
-			}
-
-			// PackageName (without PackageRoot or PackagePath) mapping (numbered suffix excluded)
-			// NameTable entries have this suffix removed and we won't match them otherwise.
-			{
-				int32 SrcPackageNamePlainStringLen = EndsWithStrippableNumber(SrcPackageName);						// PackageName_0
-				if (SrcPackageNamePlainStringLen != INDEX_NONE)
-				{
-					FStringView SrcPlainPackageName(SrcPackageName.GetData(), SrcPackageNamePlainStringLen);
-
-					// It's valid to rename a numbered name to a non-numbered name, 
-					// default to the original name and use the non-numbered name if found later
-					FStringView DstPlainPackageName = DstPackageName;
-					int32 DstPackageNamePlainStringLen = EndsWithStrippableNumber(DstPackageName);
-					if (DstPackageNamePlainStringLen != INDEX_NONE)
-					{
-						DstPlainPackageName = FStringView(DstPackageName.GetData(), DstPackageNamePlainStringLen);
-					}
-					Result.Add(FString(SrcPlainPackageName), FString(DstPlainPackageName));							// PackageName => DstPackageName
-				}
-			}
-
-			// Compiled Blueprint class names
-			{
-				NameBuilder.Reset();
-				NameBuilder.Append(SrcPackageName);
-				NameBuilder.Append(TEXT("_C"));
-				FString SrcBlueprintClassName = NameBuilder.ToString();
-
-				NameBuilder.Reset();
-				NameBuilder.Append(DstPackageName);
-				NameBuilder.Append(TEXT("_C"));
-				FString DstBlueprintClassName = NameBuilder.ToString();
-				Result.Add(MoveTemp(SrcBlueprintClassName), MoveTemp(DstBlueprintClassName));
-			}
-
-			// Blueprint generated class default object
-			{
-				NameBuilder.Reset();
-				NameBuilder.Append(DEFAULT_OBJECT_PREFIX);
-				NameBuilder.Append(SrcPackageName);
-				NameBuilder.Append(TEXT("_C"));
-				FString SrcDefaultGeneratedBlueprintClassName = NameBuilder.ToString();
-
-				NameBuilder.Reset();
-				NameBuilder.Append(DEFAULT_OBJECT_PREFIX);
-				NameBuilder.Append(DstPackageName);
-				NameBuilder.Append(TEXT("_C"));
-				FString DstDefaultGeneratedBlueprintClassName = NameBuilder.ToString();
-				Result.Add(MoveTemp(SrcDefaultGeneratedBlueprintClassName), MoveTemp(DstDefaultGeneratedBlueprintClassName));
-			}
-		}
-	}
-
-	return Result;
-}
 }
 
 bool UAssetToolsImpl::AdvancedCopyPackages(
@@ -2635,7 +2484,7 @@
 			UE::Tasks::FTaskEvent PatchAssetsCompletionTask{ UE_SOURCE_LOCATION };
 
 			TMap<FString, FString> ToCopyAndPatchPackages = AllSourceAndDestPackages(SourceAndDestPackages);
-			TMap<FString, FString> PatchingPatterns = GenerateAdditionalAssetMappings(SourceAndDestPackages);
+			TMap<FString, FString> PatchingPatterns = GetAdditionalPatchCopyMappings(SourceAndDestPackages);
 			PatchingPatterns.Append(ToCopyAndPatchPackages);
 
 			// Construct all filenames
@@ -3046,7 +2895,7 @@
 	return (Result == FAssetHeaderPatcher::EResult::Success);
 }
 
-TMap<FString, FString> UAssetToolsImpl::GetMappingsForRootPackageRename(
+TMap<FString, FString> UAssetToolsImpl::GetPatchCopyMappingsForRootRename(
 	const FString& SrcRoot,
 	const FString& DstRoot,
 	const FString& SrcBaseDir,
@@ -3116,6 +2965,173 @@
 	return Result;
 }
 
+TMap<FString, FString> UAssetToolsImpl::GetAdditionalPatchCopyMappings(const TMap<FString, FString>& SourceAndDestPackages) const
+{
+	TMap<FString, FString> Result;
+
+	TStringBuilder<24> ExternalActorsFolderBuilder;
+	ExternalActorsFolderBuilder << FPackagePath::GetExternalActorsFolderName() << TEXT("/");
+	const FStringView ExternalActorsFolder = ExternalActorsFolderBuilder.ToView();
+
+	TStringBuilder<24> ExternalObjectsFolderBuilder;
+	ExternalObjectsFolderBuilder << FPackagePath::GetExternalObjectsFolderName() << TEXT("/");
+	const FStringView ExternalObjectsFolder = ExternalObjectsFolderBuilder.ToView();
+
+	TStringBuilder<NAME_SIZE> NameBuilder;
+	for (const TTuple<FString, FString>& Package : SourceAndDestPackages)
+	{
+		const FString& SrcNameString = Package.Key;
+		const FString& DstNameString = Package.Value;
+
+		FStringView DstPackageName = FPathViews::GetBaseFilename(DstNameString);
+		FStringView SrcPackageName;
+		{
+			FStringView SrcPackageRoot;
+			FStringView SrcPackagePath;
+			if (!ensure(SplitLongPackageName(SrcNameString, SrcPackageRoot, SrcPackagePath, SrcPackageName))
+				|| SrcPackagePath.StartsWith(ExternalActorsFolder)
+				|| SrcPackagePath.StartsWith(ExternalObjectsFolder))
+			{
+				continue;
+			}
+		}
+
+		// Path.ObjectName mapping
+		{
+			NameBuilder.Reset();
+			NameBuilder.Append(SrcNameString);
+			NameBuilder.AppendChar(TEXT('.'));
+			NameBuilder.Append(SrcPackageName);
+			FString SrcObjectName = NameBuilder.ToString();
+
+			NameBuilder.Reset();
+			NameBuilder.Append(DstNameString);
+			NameBuilder.AppendChar(TEXT('.'));
+			NameBuilder.Append(DstPackageName);
+			FString DstObjectName = NameBuilder.ToString();
+			Result.Add(MoveTemp(SrcObjectName), MoveTemp(DstObjectName));
+		}
+
+		// For package names with a '_[0-9]+' suffix: Add a mapping with the suffix stripped but 
+		// maintain the full long path provided.
+		// NameTable entries have this suffix removed and we won't match them otherwise.
+		{
+			int32 SrcNameStrPlainStringLen = EndsWithStrippableNumber(SrcNameString);				// PackageRoot_0/PackagePath_0/PackageName_0
+			if (SrcNameStrPlainStringLen != INDEX_NONE)
+			{
+				FStringView SrcNamePlainString(*SrcNameString, SrcNameStrPlainStringLen);			// PackageRoot_0/PackagePath_0/PackageName
+
+				// It's valid to rename a numbered name to a non-numbered name, 
+				// default to the original name and use the non-numbered name if found later
+				FStringView DstNamePlainString = DstNameString;
+				int32 DstNameStrPlainStringLen = EndsWithStrippableNumber(DstNameString);
+				if (DstNameStrPlainStringLen != INDEX_NONE)
+				{
+					DstNamePlainString.LeftInline(DstNameStrPlainStringLen);
+				}
+				Result.Add(FString(SrcNamePlainString), FString(DstNamePlainString));				// PackageRoot_0/PackagePath_0/PackageName => DstPackageRoot_0/DstPackagePath_0/DstPackageName
+
+				FStringView SrcShortPackageName = FPathViews::GetBaseFilename(SrcNamePlainString);	// PackageName
+				NameBuilder.Reset();
+				NameBuilder.Append(SrcNameString);
+				NameBuilder.AppendChar(TEXT('.'));
+				NameBuilder.Append(SrcShortPackageName);
+				FString SrcPathWithPlainObjectName = NameBuilder.ToString();
+
+				FStringView DstShortPackageName = FPathViews::GetBaseFilename(DstNamePlainString);	// PackageName
+				NameBuilder.Reset();
+				NameBuilder.Append(DstNameString);
+				NameBuilder.AppendChar(TEXT('.'));
+				NameBuilder.Append(DstShortPackageName);
+				FString DstPathWithObjectName = NameBuilder.ToString();
+				Result.Add(MoveTemp(SrcPathWithPlainObjectName), MoveTemp(DstPathWithObjectName));	// PackageRoot_0/PackagePath_0/PackageName_0.PackageName => DstPackageRoot_0/DstPackagePath_0/DstPackageName_0.DstPackageName
+			}
+		}
+
+		// VerseAssetPath mapping used in GatherableTextData
+		// e.g. /localhost/Some/Package/Path/PackageName/PackageObject (note the use of '/' instead of '.' for the top-level Package object)
+		{
+			// Using a constant to avoid pulling in, what would otherwise be, unnecessary dependencies
+			const TCHAR* VerseRoot = TEXT("/localhost");
+
+			NameBuilder.Reset();
+			NameBuilder.Append(VerseRoot);
+			NameBuilder.Append(SrcNameString);
+			NameBuilder.AppendChar(TEXT('/'));
+			NameBuilder.Append(SrcPackageName);
+			FString SrcObjectName = NameBuilder.ToString();
+
+			NameBuilder.Reset();
+			NameBuilder.Append(VerseRoot);
+			NameBuilder.Append(DstNameString);
+			NameBuilder.AppendChar(TEXT('/'));
+			NameBuilder.Append(DstPackageName);
+			FString DstObjectName = NameBuilder.ToString();
+			Result.Add(MoveTemp(SrcObjectName), MoveTemp(DstObjectName));
+		}
+
+		if (SrcPackageName != DstPackageName)
+		{
+			// PackageName (without PackageRoot or PackagePath) mapping (numbered suffix included)
+			{
+				Result.Add({ FString(SrcPackageName), FString(DstPackageName) });
+			}
+
+			// PackageName (without PackageRoot or PackagePath) mapping (numbered suffix excluded)
+			// NameTable entries have this suffix removed and we won't match them otherwise.
+			{
+				int32 SrcPackageNamePlainStringLen = EndsWithStrippableNumber(SrcPackageName);						// PackageName_0
+				if (SrcPackageNamePlainStringLen != INDEX_NONE)
+				{
+					FStringView SrcPlainPackageName(SrcPackageName.GetData(), SrcPackageNamePlainStringLen);
+
+					// It's valid to rename a numbered name to a non-numbered name, 
+					// default to the original name and use the non-numbered name if found later
+					FStringView DstPlainPackageName = DstPackageName;
+					int32 DstPackageNamePlainStringLen = EndsWithStrippableNumber(DstPackageName);
+					if (DstPackageNamePlainStringLen != INDEX_NONE)
+					{
+						DstPlainPackageName = FStringView(DstPackageName.GetData(), DstPackageNamePlainStringLen);
+					}
+					Result.Add(FString(SrcPlainPackageName), FString(DstPlainPackageName));							// PackageName => DstPackageName
+				}
+			}
+
+			// Compiled Blueprint class names
+			{
+				NameBuilder.Reset();
+				NameBuilder.Append(SrcPackageName);
+				NameBuilder.Append(TEXT("_C"));
+				FString SrcBlueprintClassName = NameBuilder.ToString();
+
+				NameBuilder.Reset();
+				NameBuilder.Append(DstPackageName);
+				NameBuilder.Append(TEXT("_C"));
+				FString DstBlueprintClassName = NameBuilder.ToString();
+				Result.Add(MoveTemp(SrcBlueprintClassName), MoveTemp(DstBlueprintClassName));
+			}
+
+			// Blueprint generated class default object
+			{
+				NameBuilder.Reset();
+				NameBuilder.Append(DEFAULT_OBJECT_PREFIX);
+				NameBuilder.Append(SrcPackageName);
+				NameBuilder.Append(TEXT("_C"));
+				FString SrcDefaultGeneratedBlueprintClassName = NameBuilder.ToString();
+
+				NameBuilder.Reset();
+				NameBuilder.Append(DEFAULT_OBJECT_PREFIX);
+				NameBuilder.Append(DstPackageName);
+				NameBuilder.Append(TEXT("_C"));
+				FString DstDefaultGeneratedBlueprintClassName = NameBuilder.ToString();
+				Result.Add(MoveTemp(SrcDefaultGeneratedBlueprintClassName), MoveTemp(DstDefaultGeneratedBlueprintClassName));
+			}
+		}
+	}
+
+	return Result;
+}
+
 bool UAssetToolsImpl::IsDiscoveringAssetsInProgress() const
 {
 	FAssetRegistryModule& AssetRegistryModule = FModuleManager::LoadModuleChecked<FAssetRegistryModule>("AssetRegistry");
diff --git a/Engine/Source/Developer/AssetTools/Private/AssetTools.h b/Engine/Source/Developer/AssetTools/Private/AssetTools.h
--- a/Engine/Source/Developer/AssetTools/Private/AssetTools.h
+++ b/Engine/Source/Developer/AssetTools/Private/AssetTools.h
@@ -124,17 +124,15 @@ public:
 	virtual void ExpandDirectories(const TArray<FString>& Files, const FString& DestinationPath, TArray<TPair<FString, FString>>& FilesAndDestinations) const override;
 	virtual bool AdvancedCopyPackages(const FAdvancedCopyParams& CopyParams, const TArray<TMap<FString, FString>>& PackagesAndDestinations) const override;
 	virtual bool AdvancedCopyPackages(const TMap<FString, FString>& SourceAndDestPackages, const bool bForceAutosave, const bool bCopyOverAllDestinationOverlaps, FDuplicatedObjects* OutDuplicatedObjects, EMessageSeverity::Type NotificationSeverityFilter) const override;
-	virtual bool PatchCopyPackageFile(const FString& SrcFile, const FString& DstFile, const TMap<FString, FString>& SearchForAndReplace) const;
-	virtual TMap<FString, FString> GetMappingsForRootPackageRename(const FString& SrcRoot, const FString& DstRoot, const FString& SrcBaseDir, const TArray<TPair<FString, FString>>& SourceAndDestFiles, const TMap<FString, FString>& MountPointReplacements) const;
-
-
 	virtual void GenerateAdvancedCopyDestinations(FAdvancedCopyParams& InParams, const TArray<FName>& InPackageNamesToCopy, const class UAdvancedCopyCustomization* CopyCustomization, TMap<FString, FString>& OutPackagesAndDestinations) const override;
 	virtual bool FlattenAdvancedCopyDestinations(const TArray<TMap<FString, FString>>& PackagesAndDestinations, TMap<FString, FString>& FlattenedPackagesAndDestinations) const override;
 	virtual bool ValidateFlattenedAdvancedCopyDestinations(const TMap<FString, FString>& FlattenedPackagesAndDestinations) const override;
 	virtual void GetAllAdvancedCopySources(FName SelectedPackage, FAdvancedCopyParams& CopyParams, TArray<FName>& OutPackageNamesToCopy, TMap<FName, FName>& DependencyMap, const class UAdvancedCopyCustomization* CopyCustomization) const override;
 	virtual void InitAdvancedCopyFromCopyParams(FAdvancedCopyParams CopyParams) const override;
+	virtual bool PatchCopyPackageFile(const FString& SrcFile, const FString& DstFile, const TMap<FString, FString>& SearchForAndReplace) const override;
+	virtual TMap<FString, FString> GetPatchCopyMappingsForRootRename(const FString& SrcRoot, const FString& DstRoot, const FString& SrcBaseDir, const TArray<TPair<FString, FString>>& SourceAndDestFiles, const TMap<FString, FString>& MountPointReplacements) const override;
+	virtual TMap<FString, FString> GetAdditionalPatchCopyMappings(const TMap<FString, FString>& SourceAndDestPackages) const override;
 	virtual void OpenEditorForAssets(const TArray<UObject*>& Assets) override;
-
 	virtual void ConvertVirtualTextures(const TArray<UTexture2D*>& Textures, bool bConvertBackToNonVirtual, const TArray<UMaterial*>* RelatedMaterials = nullptr) const override;
 	virtual bool IsAssetClassSupported(const UClass* AssetClass) const override;
 	virtual TArray<UFactory*> GetNewAssetFactories() const override;
diff --git a/Engine/Source/Developer/AssetTools/Public/IAssetTools.h b/Engine/Source/Developer/AssetTools/Public/IAssetTools.h
--- a/Engine/Source/Developer/AssetTools/Public/IAssetTools.h
+++ b/Engine/Source/Developer/AssetTools/Public/IAssetTools.h
@@ -569,12 +569,6 @@ public:
 	/** Copies files after the flattened map of sources and destinations was confirmed */
 	virtual bool AdvancedCopyPackages(const TMap<FString, FString>& SourceAndDestPackages, const bool bForceAutosave = false, const bool bCopyOverAllDestinationOverlaps = true, FDuplicatedObjects* OutDuplicatedObjects = nullptr, EMessageSeverity::Type NotificationSeverityFilter = EMessageSeverity::Info) const = 0;
 
-	/** Copies a file, patching internal references without performing a de-serialization. This is a blocking operation. returns true on successful copy */
-	virtual bool PatchCopyPackageFile(const FString& SrcFile, const FString& DstFile, const TMap<FString, FString>& SearchForAndReplace) const = 0;
-
-	/** Generates the SearchAndReplace map for a PatchCopyPackageFile if all you are doing is changing the root and not the name. */
-	virtual TMap<FString, FString> GetMappingsForRootPackageRename(const FString& SrcRoot, const FString& DstRoot, const FString& SrcBaseDir, const TArray<TPair<FString, FString>>& SourceAndDestFiles, const TMap<FString, FString>& MountPointReplacements) const = 0;
-
 	/* Given a set of packages to copy, generate the map of those packages to destination filenames */
 	virtual void GenerateAdvancedCopyDestinations(FAdvancedCopyParams& InParams, const TArray<FName>& InPackageNamesToCopy, const UAdvancedCopyCustomization* CopyCustomization, TMap<FString, FString>& OutPackagesAndDestinations) const = 0;
 
@@ -590,6 +584,15 @@ public:
 	/* Given a complete set of copy parameters, which includes the selected package set, start the advanced copy process */
 	virtual void InitAdvancedCopyFromCopyParams(FAdvancedCopyParams CopyParams) const = 0;
 
+	/** Copies a file, patching internal references without performing a de-serialization. This is a blocking operation. returns true on successful copy */
+	UE_INTERNAL virtual bool PatchCopyPackageFile(const FString& SrcFile, const FString& DstFile, const TMap<FString, FString>& SearchForAndReplace) const = 0;
+
+	/** Generates the PatchCopyPackageFile SearchForAndReplace parameter if all you are doing is changing the root and not the relative path of assets */
+	UE_INTERNAL virtual TMap<FString, FString> GetPatchCopyMappingsForRootRename(const FString& SrcRoot, const FString& DstRoot, const FString& SrcBaseDir, const TArray<TPair<FString, FString>>& SourceAndDestFiles, const TMap<FString, FString>& MountPointReplacements) const = 0;
+
+	/** Generates additional entries for the PatchCopyPackageFile SearchForAndReplace parameter based on the specified package paths mapping */
+	UE_INTERNAL virtual TMap<FString, FString> GetAdditionalPatchCopyMappings(const TMap<FString, FString>& SourceAndDestPackages) const = 0;
+
 	/** Opens editor for assets */
 	UFUNCTION(BlueprintCallable, Category = "Editor Scripting | Asset Tools", meta = (DeprecatedFunction, DeprecationMessage = "Please use UAssetEditorSubsystem::OpenEditorForAssets instead."))
 	virtual void OpenEditorForAssets(const TArray<UObject*>& Assets) = 0;
