[World Partition] Remove assert when trying to use advanced copy on a partitioned world with streaming enabled, since it now works.

diff --git a/Engine/Source/Developer/AssetTools/Private/AssetTools.cpp b/Engine/Source/Developer/AssetTools/Private/AssetTools.cpp
--- a/Engine/Source/Developer/AssetTools/Private/AssetTools.cpp
+++ b/Engine/Source/Developer/AssetTools/Private/AssetTools.cpp
@@ -2668,7 +2668,6 @@ bool UAssetToolsImpl::AdvancedCopyPackages(
 							{
 								if (UWorldPartition* WorldPartition = World->GetWorldPartition())
 								{
-									check(!WorldPartition->IsStreamingEnabled());
 									if (!WorldPartition->IsInitialized())
 									{
 										WorldPartition->Initialize(World, FTransform::Identity);
