Undo CL: 35599474
Static initialization on startup

diff --git a/Engine/Plugins/Editor/EngineAssetDefinitions/Source/EngineAssetDefinitions.Build.cs b/Engine/Plugins/Editor/EngineAssetDefinitions/Source/EngineAssetDefinitions.Build.cs
--- a/Engine/Plugins/Editor/EngineAssetDefinitions/Source/EngineAssetDefinitions.Build.cs
+++ b/Engine/Plugins/Editor/EngineAssetDefinitions/Source/EngineAssetDefinitions.Build.cs
@@ -23,7 +23,6 @@ public class EngineAssetDefinitions : ModuleRules
 		PrivateDependencyModuleNames.AddRange(
 			new string[]
 			{
-				"ApplicationCore",
 				"ContentBrowser",
 				"ContentBrowserData",
 				"AssetDefinition",
diff --git a/Engine/Plugins/Editor/EngineAssetDefinitions/Source/Private/AssetDefinition_ForceFeedbackEffect.cpp b/Engine/Plugins/Editor/EngineAssetDefinitions/Source/Private/AssetDefinition_ForceFeedbackEffect.cpp
deleted file mode 100644
diff --git a/Engine/Plugins/Editor/EngineAssetDefinitions/Source/Private/AssetDefinition_ForceFeedbackEffect.h b/Engine/Plugins/Editor/EngineAssetDefinitions/Source/Private/AssetDefinition_ForceFeedbackEffect.h
deleted file mode 100644
diff --git a/Engine/Plugins/Experimental/AnimToTexture/Source/AnimToTextureEditor/AnimToTextureEditor.Build.cs b/Engine/Plugins/Experimental/AnimToTexture/Source/AnimToTextureEditor/AnimToTextureEditor.Build.cs
--- a/Engine/Plugins/Experimental/AnimToTexture/Source/AnimToTextureEditor/AnimToTextureEditor.Build.cs
+++ b/Engine/Plugins/Experimental/AnimToTexture/Source/AnimToTextureEditor/AnimToTextureEditor.Build.cs
@@ -25,7 +25,6 @@ public class AnimToTextureEditor : ModuleRules
 		PublicDependencyModuleNames.AddRange(
 			new string[]
 			{
-				"AssetDefinition",
 				"Core",
 				"MaterialEditor",
 				"MessageLog"
@@ -37,8 +36,6 @@ public class AnimToTextureEditor : ModuleRules
 		PrivateDependencyModuleNames.AddRange(
 			new string[]
 			{
-				"AssetTools",
-				"ContentBrowser",
 				"CoreUObject",
 				"Engine",
 				"Slate",
@@ -47,7 +44,6 @@ public class AnimToTextureEditor : ModuleRules
 				"RawMesh",
 				"MeshDescription",
 				"StaticMeshDescription",
-				"ToolMenus",
 				"UnrealEd",
 			}
 			);
diff --git a/Engine/Plugins/Experimental/AnimToTexture/Source/AnimToTextureEditor/Private/AnimToTextureAssetActions.cpp b/Engine/Plugins/Experimental/AnimToTexture/Source/AnimToTextureEditor/Private/AnimToTextureAssetActions.cpp
new file mode 100644
--- /dev/null
+++ b/Engine/Plugins/Experimental/AnimToTexture/Source/AnimToTextureEditor/Private/AnimToTextureAssetActions.cpp
@@ -0,0 +1,97 @@
+// Copyright Epic Games, Inc. All Rights Reserved.
+
+#include "AnimToTextureAssetActions.h"
+#include "AnimToTextureDataAsset.h"
+#include "AnimToTextureBPLibrary.h"
+
+#include "Materials/MaterialInstanceConstant.h"
+#include "Styling/SlateIconFinder.h"
+#include "Widgets/Images/SImage.h"
+#include "Widgets/Layout/SBorder.h"
+
+FText FAnimToTextureAssetActions::GetName() const
+{
+	return NSLOCTEXT("AssetTypeActions", "AnimToTextureAssetActions", "AnimToTexture");
+}
+
+FColor FAnimToTextureAssetActions::GetTypeColor() const
+{
+	return FColor::Blue;
+}
+
+TSharedPtr<SWidget> FAnimToTextureAssetActions::GetThumbnailOverlay(const FAssetData& AssetData) const
+{
+	const FSlateBrush* Icon = FSlateIconFinder::FindIconBrushForClass(UAnimToTextureDataAsset::StaticClass());;
+
+	return SNew(SBorder)
+		.BorderImage(FAppStyle::GetNoBrush())
+		.Visibility(EVisibility::HitTestInvisible)
+		.Padding(FMargin(0.0f, 0.0f, 0.0f, 3.0f))
+		.HAlign(HAlign_Right)
+		.VAlign(VAlign_Bottom)
+		[
+			SNew(SImage)
+			.Image(Icon)
+		];
+}
+
+UClass* FAnimToTextureAssetActions::GetSupportedClass() const
+{
+	return UAnimToTextureDataAsset::StaticClass();
+}
+
+uint32 FAnimToTextureAssetActions::GetCategories()
+{
+	return EAssetTypeCategories::Animation;
+}
+
+const TArray<FText>& FAnimToTextureAssetActions::GetSubMenus() const
+{
+	static const TArray<FText> SubMenus
+	{
+		NSLOCTEXT("AnimToTextureAssetActions", "AnimDeformersSubMenu", "Deformers")
+	};
+	return SubMenus;
+}
+
+void FAnimToTextureAssetActions::GetActions( const TArray<UObject*>& InObjects, FMenuBuilder& MenuBuilder )
+{
+	TArray<TWeakObjectPtr<UAnimToTextureDataAsset>> Objects = GetTypedWeakObjectPtrs<UAnimToTextureDataAsset>(InObjects);
+
+	MenuBuilder.AddMenuEntry(
+		NSLOCTEXT("AnimToTextureAssetActions","AnimToTexture_Run", "Run Animation To Texture"),
+		NSLOCTEXT("AnimToTextureAssetActions", "AnimToTexture_RunTooltip", "Creates Vertex Animation Textures (VAT)"),
+		FSlateIcon(),
+		FUIAction(
+			FExecuteAction::CreateSP( this, &FAnimToTextureAssetActions::RunAnimToTexture, Objects ),
+			FCanExecuteAction()
+			)
+		);
+
+	MenuBuilder.AddMenuSeparator();
+}
+
+void FAnimToTextureAssetActions::RunAnimToTexture(TArray<TWeakObjectPtr<UAnimToTextureDataAsset>> Objects)
+{
+	for (auto ObjIt = Objects.CreateConstIterator(); ObjIt; ++ObjIt)
+	{
+		if (UAnimToTextureDataAsset* DataAsset = (*ObjIt).Get())
+		{
+			// Create UVs and Textures
+			if (UAnimToTextureBPLibrary::AnimationToTexture(DataAsset))
+			{
+				// Update Material Instances (if Possible)
+				if (UStaticMesh* StaticMesh = DataAsset->GetStaticMesh())
+				{
+					for (FStaticMaterial& StaticMaterial: StaticMesh->GetStaticMaterials())
+					{
+						if (UMaterialInstanceConstant* MaterialInstanceConstant = Cast<UMaterialInstanceConstant>(StaticMaterial.MaterialInterface))
+						{
+							UAnimToTextureBPLibrary::UpdateMaterialInstanceFromDataAsset(DataAsset, MaterialInstanceConstant, EMaterialParameterAssociation::LayerParameter);
+						}
+					}
+				}
+			}
+		}
+	}
+}
diff --git a/Engine/Plugins/Experimental/AnimToTexture/Source/AnimToTextureEditor/Private/AnimToTextureEditorModule.cpp b/Engine/Plugins/Experimental/AnimToTexture/Source/AnimToTextureEditor/Private/AnimToTextureEditorModule.cpp
--- a/Engine/Plugins/Experimental/AnimToTexture/Source/AnimToTextureEditor/Private/AnimToTextureEditorModule.cpp
+++ b/Engine/Plugins/Experimental/AnimToTexture/Source/AnimToTextureEditor/Private/AnimToTextureEditorModule.cpp
@@ -1,6 +1,7 @@
 // Copyright Epic Games, Inc. All Rights Reserved.
 
 #include "AnimToTextureEditorModule.h"
+#include "AnimToTextureAssetActions.h"
 
 #include "AssetToolsModule.h"
 #include "MessageLogModule.h"
@@ -11,6 +12,11 @@ DEFINE_LOG_CATEGORY(LogAnimToTextureEditor);
 
 void FAnimToTextureEditorModule::StartupModule()
 {
+	// Asset actions
+	FAssetToolsModule& AssetToolsModule = FModuleManager::LoadModuleChecked< FAssetToolsModule >("AssetTools");	
+	TSharedPtr<FAnimToTextureAssetActions> AnimToTextureTypeActions = MakeShareable( new FAnimToTextureAssetActions);
+	AssetToolsModule.Get().RegisterAssetTypeActions(AnimToTextureTypeActions.ToSharedRef() );
+
 	// Register Log
 	FMessageLogModule& MessageLogModule = FModuleManager::LoadModuleChecked<FMessageLogModule>("MessageLog");
 	FMessageLogInitializationOptions InitOptions;
diff --git a/Engine/Plugins/Experimental/AnimToTexture/Source/AnimToTextureEditor/Private/AssetDefinition_AnimToTexture.cpp b/Engine/Plugins/Experimental/AnimToTexture/Source/AnimToTextureEditor/Private/AssetDefinition_AnimToTexture.cpp
deleted file mode 100644
diff --git a/Engine/Plugins/Experimental/AnimToTexture/Source/AnimToTextureEditor/Public/AnimToTextureAssetActions.h b/Engine/Plugins/Experimental/AnimToTexture/Source/AnimToTextureEditor/Public/AnimToTextureAssetActions.h
new file mode 100644
--- /dev/null
+++ b/Engine/Plugins/Experimental/AnimToTexture/Source/AnimToTextureEditor/Public/AnimToTextureAssetActions.h
@@ -0,0 +1,30 @@
+// Copyright Epic Games, Inc. All Rights Reserved.
+
+#pragma once
+
+#include "AssetTypeActions_Base.h"
+#include "AnimToTextureDataAsset.h"
+
+class FMenuBuilder;
+class IToolkitHost;
+class UClass;
+class UObject;
+
+
+class FAnimToTextureAssetActions : public FAssetTypeActions_Base
+{
+public:
+	
+	// IAssetTypeActions Implementation
+	FText GetName() const override;
+	FColor GetTypeColor() const override;
+	UClass* GetSupportedClass() const override;	
+	void GetActions( const TArray<UObject*>& InObjects, FMenuBuilder& MenuBuilder ) override;
+	uint32 GetCategories() override;
+	const TArray<FText>& GetSubMenus() const override;
+	TSharedPtr<SWidget> GetThumbnailOverlay(const FAssetData& AssetData) const override;
+
+private:
+	void RunAnimToTexture(TArray<TWeakObjectPtr<UAnimToTextureDataAsset>> Objects);
+	
+};
diff --git a/Engine/Plugins/Experimental/AnimToTexture/Source/AnimToTextureEditor/Public/AssetDefinition_AnimToTexture.h b/Engine/Plugins/Experimental/AnimToTexture/Source/AnimToTextureEditor/Public/AssetDefinition_AnimToTexture.h
deleted file mode 100644
diff --git a/Engine/Plugins/Runtime/Metasound/Source/MetasoundEditor/Private/MetasoundAssetDefinitions.cpp b/Engine/Plugins/Runtime/Metasound/Source/MetasoundEditor/Private/MetasoundAssetDefinitions.cpp
--- a/Engine/Plugins/Runtime/Metasound/Source/MetasoundEditor/Private/MetasoundAssetDefinitions.cpp
+++ b/Engine/Plugins/Runtime/Metasound/Source/MetasoundEditor/Private/MetasoundAssetDefinitions.cpp
@@ -295,72 +295,6 @@ TSharedPtr<SWidget> UAssetDefinition_MetaSoundSource::GetThumbnailOverlay(const
 	return UAssetDefinition_SoundBase::GetSoundBaseThumbnailOverlay(InAssetData, MoveTemp(OnClickedLambdaOverride));
 }
 
-bool UAssetDefinition_MetaSoundSource::GetThumbnailActionOverlay(const FAssetData& InAssetData, FAssetActionThumbnailOverlayInfo& OutActionOverlayInfo) const
-{
-	auto OnGetDisplayBrushLambda = [InAssetData]() -> const FSlateBrush*
-	{
-		if (UE::AudioEditor::IsSoundPlaying(InAssetData))
-		{
-			return FAppStyle::GetBrush("ContentBrowser.AssetAction.StopIcon");
-		}
-
-		return FAppStyle::GetBrush("ContentBrowser.AssetAction.PlayIcon");
-	};
-
-	OutActionOverlayInfo.ActionImageWidget = SNew(SImage).Image_Lambda(OnGetDisplayBrushLambda);
-
-	auto OnToolTipTextLambda = [InAssetData]() -> FText
-	{
-		if (UE::AudioEditor::IsSoundPlaying(InAssetData))
-		{
-			return LOCTEXT("Thumbnail_StopSoundToolTip", "Stop selected sound");
-		}
-
-		return LOCTEXT("Thumbnail_PlaySoundToolTip", "Play selected sound");
-	};
-
-	auto OnClickedLambda = [InAssetData]() -> FReply
-	{
-		TSharedPtr<Metasound::Editor::FEditor> Editor = Metasound::Editor::FGraphBuilder::GetEditorForMetasound(*InAssetData.GetAsset());
-		if (UE::AudioEditor::IsSoundPlaying(InAssetData))
-		{
-			if (Editor.IsValid())
-			{
-				Editor->Stop();
-			}
-			else
-			{
-				UE::AudioEditor::StopSound();
-			}
-		}
-		else
-		{
-			if (Editor.IsValid())
-			{
-				Editor->Play();
-			}
-			else
-			{
-				// Load and play sound
-				UE::AudioEditor::PlaySound(Cast<USoundBase>(InAssetData.GetAsset()));
-			}
-		}
-		return FReply::Handled();
-	};
-
-	OutActionOverlayInfo.ActionButtonWidget = SNew(SButton)
-		.ButtonStyle(FAppStyle::Get(), "HoverHintOnly")
-		.ContentPadding(0.0f)
-		.ToolTipText_Lambda(OnToolTipTextLambda)
-		.OnClicked_Lambda(OnClickedLambda)
-		[
-			SNew(SImage)
-			.Image_Lambda(OnGetDisplayBrushLambda)
-		];
-
-	return true;
-}
-
 EAssetCommandResult UAssetDefinition_MetaSoundSource::ActivateAssets(const FAssetActivateArgs& ActivateArgs) const
 {
 	if (ActivateArgs.ActivationMethod == EAssetActivationMethod::Previewed)
diff --git a/Engine/Plugins/Runtime/Metasound/Source/MetasoundEditor/Private/MetasoundAssetDefinitions.h b/Engine/Plugins/Runtime/Metasound/Source/MetasoundEditor/Private/MetasoundAssetDefinitions.h
--- a/Engine/Plugins/Runtime/Metasound/Source/MetasoundEditor/Private/MetasoundAssetDefinitions.h
+++ b/Engine/Plugins/Runtime/Metasound/Source/MetasoundEditor/Private/MetasoundAssetDefinitions.h
@@ -38,7 +38,6 @@ public:
 	virtual EAssetCommandResult OpenAssets(const FAssetOpenArgs& OpenArgs) const override;
 
 	virtual TSharedPtr<SWidget> GetThumbnailOverlay(const FAssetData& InAssetData) const override;
-	virtual bool GetThumbnailActionOverlay(const FAssetData& InAssetData, FAssetActionThumbnailOverlayInfo& OutActionOverlayInfo) const override;
 	virtual EAssetCommandResult ActivateAssets(const FAssetActivateArgs& ActivateArgs) const override;
 
  	virtual const FSlateBrush* GetThumbnailBrush(const FAssetData& InAssetData, const FName InClassName) const override;
diff --git a/Engine/Source/Developer/AssetTools/Private/AssetDefinition_AssetTypeActionsProxy.cpp b/Engine/Source/Developer/AssetTools/Private/AssetDefinition_AssetTypeActionsProxy.cpp
--- a/Engine/Source/Developer/AssetTools/Private/AssetDefinition_AssetTypeActionsProxy.cpp
+++ b/Engine/Source/Developer/AssetTools/Private/AssetDefinition_AssetTypeActionsProxy.cpp
@@ -237,14 +237,7 @@ const FSlateBrush* UAssetDefinition_AssetTypeActionsProxy::GetIconBrush(const FA
 
 TSharedPtr<SWidget> UAssetDefinition_AssetTypeActionsProxy::GetThumbnailOverlay(const FAssetData& InAssetData) const
 {
-	PRAGMA_DISABLE_DEPRECATION_WARNINGS
 	return AssetType->GetThumbnailOverlay(InAssetData);
-	PRAGMA_ENABLE_DEPRECATION_WARNINGS
 }
 
-bool UAssetDefinition_AssetTypeActionsProxy::GetThumbnailActionOverlay(const FAssetData& InAssetData, FAssetActionThumbnailOverlayInfo& OutActionOverlayInfo) const
-{
-	return false;
-}
-
-#undef LOCTEXT_NAMESPACE
+#undef LOCTEXT_NAMESPACE
\ No newline at end of file
diff --git a/Engine/Source/Developer/AssetTools/Private/AssetDefinition_AssetTypeActionsProxy.h b/Engine/Source/Developer/AssetTools/Private/AssetDefinition_AssetTypeActionsProxy.h
--- a/Engine/Source/Developer/AssetTools/Private/AssetDefinition_AssetTypeActionsProxy.h
+++ b/Engine/Source/Developer/AssetTools/Private/AssetDefinition_AssetTypeActionsProxy.h
@@ -51,7 +51,6 @@ public:
 	virtual const FSlateBrush* GetThumbnailBrush(const FAssetData& InAssetData, const FName InClassName) const override;
 	virtual const FSlateBrush* GetIconBrush(const FAssetData& InAssetData, const FName InClassName) const override;
 	virtual TSharedPtr<SWidget> GetThumbnailOverlay(const FAssetData& InAssetData) const override;
-	virtual bool GetThumbnailActionOverlay(const FAssetData& InAssetData, FAssetActionThumbnailOverlayInfo& OutActionOverlayInfo) const override;
 
 protected:
 	virtual bool CanRegisterStatically() const override { return false; }
diff --git a/Engine/Source/Developer/AssetTools/Private/AssetTools.cpp b/Engine/Source/Developer/AssetTools/Private/AssetTools.cpp
--- a/Engine/Source/Developer/AssetTools/Private/AssetTools.cpp
+++ b/Engine/Source/Developer/AssetTools/Private/AssetTools.cpp
@@ -39,6 +39,7 @@
 #include "AssetTypeActions/AssetTypeActions_AnimBoneCompressionSettings.h"
 #include "AssetTypeActions/AssetTypeActions_AnimCurveCompressionSettings.h"
 #include "AssetTypeActions/AssetTypeActions_VariableFrameStrippingSettings.h"
+#include "AssetTypeActions/AssetTypeActions_ForceFeedbackEffect.h"
 #include "AssetTypeActions/AssetTypeActions_ParticleSystem.h"
 #include "AssetTypeActions/AssetTypeActions_PhysicalMaterialMask.h"
 #include "AssetTypeActions/AssetTypeActions_DeletedObjectPlaceholder.h"
@@ -395,9 +396,7 @@ public:
 	/** Optionally returns a custom widget to overlay on top of this assets' thumbnail */
     virtual TSharedPtr<class SWidget> GetThumbnailOverlay(const FAssetData& AssetData) const override
     {
-		PRAGMA_DISABLE_DEPRECATION_WARNINGS
     	return AssetDefinitionPtr.Get()->GetThumbnailOverlay(AssetData);
-		PRAGMA_ENABLE_DEPRECATION_WARNINGS
     }
 	
 	virtual EThumbnailPrimType GetDefaultThumbnailPrimitiveType(UObject* Asset) const override
@@ -1400,10 +1399,13 @@ UAssetToolsImpl::UAssetToolsImpl(const FObjectInitializer& ObjectInitializer)
 	AllocatedCategoryBits.Add(TEXT("Texture"), FAdvancedAssetCategory(EAssetTypeCategories::Textures, LOCTEXT("TextureAssetCategory", "Texture")));
 	AllocatedCategoryBits.Add(TEXT("World"), FAdvancedAssetCategory(EAssetTypeCategories::World, LOCTEXT("WorldAssetCategory", "World")));
 
+	EAssetTypeCategories::Type InputCategoryBit = RegisterAdvancedAssetCategory(FName(TEXT("Input")), LOCTEXT("InputAssetsCategory", "Input"));
+	
 	// Register the built-in asset type actions
 	RegisterAssetTypeActions(MakeShareable(new FAssetTypeActions_AnimBoneCompressionSettings));
 	RegisterAssetTypeActions(MakeShareable(new FAssetTypeActions_AnimCurveCompressionSettings));
 	RegisterAssetTypeActions(MakeShareable(new FAssetTypeActions_VariableFrameStrippingSettings));
+	RegisterAssetTypeActions(MakeShareable(new FAssetTypeActions_ForceFeedbackEffect(InputCategoryBit)));
 	RegisterAssetTypeActions(MakeShareable(new FAssetTypeActions_ParticleSystem));
 	RegisterAssetTypeActions(MakeShareable(new FAssetTypeActions_PhysicalMaterialMask));
 	RegisterAssetTypeActions(MakeShareable(new FAssetTypeActions_DeletedObjectPlaceholder));
diff --git a/Engine/Source/Developer/AssetTools/Private/AssetTypeActions/AssetTypeActions_ForceFeedbackEffect.cpp b/Engine/Source/Developer/AssetTools/Private/AssetTypeActions/AssetTypeActions_ForceFeedbackEffect.cpp
new file mode 100644
--- /dev/null
+++ b/Engine/Source/Developer/AssetTools/Private/AssetTypeActions/AssetTypeActions_ForceFeedbackEffect.cpp
@@ -0,0 +1,273 @@
+// Copyright Epic Games, Inc. All Rights Reserved.
+
+#include "AssetTypeActions/AssetTypeActions_ForceFeedbackEffect.h"
+#include "ToolMenus.h"
+#include "Styling/AppStyle.h"
+#include "GenericPlatform/IInputInterface.h"
+#include "AssetRegistry/AssetData.h"
+#include "Framework/Application/SlateApplication.h"
+#include "Widgets/Layout/SBox.h"
+#include "Widgets/Input/SButton.h"
+#include "Widgets/Images/SImage.h"
+#include "GenericPlatform/GenericPlatformInputDeviceMapper.h"
+
+#define LOCTEXT_NAMESPACE "AssetTypeActions"
+
+UClass* FAssetTypeActions_ForceFeedbackEffect::GetSupportedClass() const
+{
+	return UForceFeedbackEffect::StaticClass();
+}
+
+void FAssetTypeActions_ForceFeedbackEffect::GetActions(const TArray<UObject*>& InObjects, FToolMenuSection& Section)
+{
+	TArray<TWeakObjectPtr<UForceFeedbackEffect>> Effects = GetTypedWeakObjectPtrs<UForceFeedbackEffect>(InObjects);
+
+	Section.AddMenuEntry(
+		"ForceFeedbackEffect_PlayEffect",
+		LOCTEXT("ForceFeedbackEffect_PlayEffect", "Play"),
+		LOCTEXT("ForceFeedbackEffect_PlayEffectTooltip", "Plays the selected force feedback effect."),
+		FSlateIcon(FAppStyle::GetAppStyleSetName(), "MediaAsset.AssetActions.Play.Small"),
+		FUIAction(
+			FExecuteAction::CreateSP( this, &FAssetTypeActions_ForceFeedbackEffect::ExecutePlayEffect, Effects ),
+			FCanExecuteAction::CreateSP( this, &FAssetTypeActions_ForceFeedbackEffect::CanExecutePlayCommand, Effects )
+			)
+		);
+
+	Section.AddMenuEntry(
+		"ForceFeedbackEffect_StopEffect",
+		LOCTEXT("ForceFeedbackEffect_StopEffect", "Stop"),
+		LOCTEXT("ForceFeedbackEffect_StopEffectTooltip", "Stops the selected force feedback effect."),
+		FSlateIcon(FAppStyle::GetAppStyleSetName(), "MediaAsset.AssetActions.Stop.Small"),
+		FUIAction(
+			FExecuteAction::CreateSP( this, &FAssetTypeActions_ForceFeedbackEffect::ExecuteStopEffect, Effects ),
+			FCanExecuteAction()
+			)
+		);
+}
+
+bool FAssetTypeActions_ForceFeedbackEffect::IsEffectPlaying(const TArray<TWeakObjectPtr<UForceFeedbackEffect>>& Objects) const
+{
+	if (PreviewForceFeedbackEffect.ForceFeedbackEffect)
+	{
+		for (const TWeakObjectPtr<UForceFeedbackEffect>& EffectPtr : Objects)
+		{
+			UForceFeedbackEffect* Effect = EffectPtr.Get();
+			if (Effect && PreviewForceFeedbackEffect.ForceFeedbackEffect == Effect)
+			{
+				return true;
+			}
+		}
+	}
+
+	return false;
+}
+
+bool FAssetTypeActions_ForceFeedbackEffect::IsEffectPlaying(const UForceFeedbackEffect* ForceFeedbackEffect) const
+{
+	return PreviewForceFeedbackEffect.ForceFeedbackEffect && PreviewForceFeedbackEffect.ForceFeedbackEffect == ForceFeedbackEffect;
+}
+
+bool FAssetTypeActions_ForceFeedbackEffect::IsEffectPlaying(const FAssetData& AssetData) const
+{
+	if (PreviewForceFeedbackEffect.ForceFeedbackEffect)
+	{
+		if (PreviewForceFeedbackEffect.ForceFeedbackEffect->GetFName() == AssetData.AssetName)
+		{
+			if (PreviewForceFeedbackEffect.ForceFeedbackEffect->GetOutermost()->GetFName() == AssetData.PackageName)
+			{
+				return true;
+			}
+		}
+	}
+
+	return false;
+}
+
+bool FAssetTypeActions_ForceFeedbackEffect::CanExecutePlayCommand(TArray<TWeakObjectPtr<UForceFeedbackEffect>> Objects) const
+{
+	return Objects.Num() == 1;
+}
+
+bool FAssetTypeActions_ForceFeedbackEffect::AssetsActivatedOverride( const TArray<UObject*>& InObjects, EAssetTypeActivationMethod::Type ActivationType )
+{
+	if (ActivationType == EAssetTypeActivationMethod::Previewed)
+	{
+		for (UObject* Object : InObjects)
+		{
+			UForceFeedbackEffect* TargetEffect = Cast<UForceFeedbackEffect>(Object);
+			if ( TargetEffect )
+			{
+				// Only target the first valid effect
+				TArray<TWeakObjectPtr<UForceFeedbackEffect>> EffectList;
+				EffectList.Add(MakeWeakObjectPtr(TargetEffect));
+				if (IsEffectPlaying(EffectList))
+				{
+					ExecuteStopEffect(EffectList);
+				}
+				else
+				{
+					ExecutePlayEffect(EffectList);
+				}
+
+				return true;
+			}
+		}
+	}
+	return false;
+}
+
+void FAssetTypeActions_ForceFeedbackEffect::ExecutePlayEffect(TArray<TWeakObjectPtr<UForceFeedbackEffect>> Objects)
+{
+	for (const TWeakObjectPtr<UForceFeedbackEffect>& EffectPtr : Objects)
+	{
+		UForceFeedbackEffect* Effect = EffectPtr.Get();
+		if (Effect)
+		{
+			// Only play the first valid effect
+			PlayEffect(Effect);
+			break;
+		}
+	}
+}
+
+void FAssetTypeActions_ForceFeedbackEffect::ExecuteStopEffect(TArray<TWeakObjectPtr<UForceFeedbackEffect>> Objects)
+{
+	StopEffect();
+}
+
+void FAssetTypeActions_ForceFeedbackEffect::PlayEffect(UForceFeedbackEffect* Effect)
+{
+	if (Effect)
+	{
+		PreviewForceFeedbackEffect.ForceFeedbackEffect = Effect;
+		PreviewForceFeedbackEffect.PlayTime = 0.f;
+		PreviewForceFeedbackEffect.PlatformUser = IPlatformInputDeviceMapper::Get().GetPrimaryPlatformUser();
+		PreviewForceFeedbackEffect.ActivateDeviceProperties();
+	}
+	else
+	{
+		StopEffect();
+	}
+}
+
+void FAssetTypeActions_ForceFeedbackEffect::StopEffect() 
+{
+	PreviewForceFeedbackEffect.ResetDeviceProperties();
+	PreviewForceFeedbackEffect.ForceFeedbackEffect = nullptr;
+
+	IInputInterface* InputInterface = FSlateApplication::Get().GetInputInterface();
+	if (InputInterface)
+	{
+		InputInterface->SetForceFeedbackChannelValues(0, FForceFeedbackValues());
+	}
+}
+
+TSharedPtr<SWidget> FAssetTypeActions_ForceFeedbackEffect::GetThumbnailOverlay(const FAssetData& AssetData) const
+{
+	auto OnGetDisplayBrushLambda = [this, AssetData]() -> const FSlateBrush*
+	{
+		if (IsEffectPlaying(AssetData))
+		{
+			return FAppStyle::GetBrush("MediaAsset.AssetActions.Stop.Large");
+		}
+
+		return FAppStyle::GetBrush("MediaAsset.AssetActions.Play.Large");
+	};
+
+	FAssetTypeActions_ForceFeedbackEffect* MutableThis = const_cast<FAssetTypeActions_ForceFeedbackEffect*>(this);
+	auto OnClickedLambda = [MutableThis, AssetData]() -> FReply
+	{
+		if (MutableThis->IsEffectPlaying(AssetData))
+		{
+			MutableThis->StopEffect();
+		}
+		else
+		{
+			// Load and play asset
+			MutableThis->PlayEffect(Cast<UForceFeedbackEffect>(AssetData.GetAsset()));
+		}
+		return FReply::Handled();
+	};
+
+	auto OnToolTipTextLambda = [this, AssetData]() -> FText
+	{
+		if (IsEffectPlaying(AssetData))
+		{
+			return LOCTEXT("Thumbnail_StopForceFeedbackToolTip", "Stop selected force feedback effect");
+		}
+
+		return LOCTEXT("Thumbnail_PlayForceFeedbackToolTip", "Play selected force feedback effect");
+	};
+
+	TSharedRef<SBox> Box = SNew(SBox)
+		.HAlign(HAlign_Center)
+		.VAlign(VAlign_Center)
+		.Padding(FMargin(2));
+
+	auto OnGetVisibilityLambda = [this, Box, AssetData]() -> EVisibility
+	{
+		if (Box->IsHovered() || IsEffectPlaying(AssetData))
+		{
+			return EVisibility::Visible;
+		}
+
+		return EVisibility::Hidden;
+	};
+
+	TSharedRef<SButton> BoxContent = SNew(SButton)
+		.ButtonStyle(FAppStyle::Get(), "HoverHintOnly")
+		.ToolTipText_Lambda(OnToolTipTextLambda)
+		.Cursor(EMouseCursor::Default) // The outer widget can specify a DragHand cursor, so we need to override that here
+		.ForegroundColor(FSlateColor::UseForeground())
+		.IsFocusable(false)
+		.OnClicked_Lambda(OnClickedLambda)
+		.Visibility_Lambda(OnGetVisibilityLambda)
+		[
+			SNew(SBox)
+			.MinDesiredWidth(16.f)
+			.MinDesiredHeight(16.f)
+			[
+				SNew(SImage)
+				.Image_Lambda(OnGetDisplayBrushLambda)
+			]
+		];
+
+	Box->SetContent(BoxContent);
+	Box->SetVisibility(EVisibility::Visible);
+
+	return Box;
+}
+
+bool FPreviewForceFeedbackEffect::IsTickable() const
+{
+	return (ForceFeedbackEffect != nullptr);
+}
+
+void FPreviewForceFeedbackEffect::Tick( float DeltaTime )
+{
+	FForceFeedbackValues ForceFeedbackValues;
+
+	if (!Update(DeltaTime, ForceFeedbackValues))
+	{
+		ResetDeviceProperties();
+		ForceFeedbackEffect = nullptr;
+	}
+
+	IInputInterface* InputInterface = FSlateApplication::Get().GetInputInterface();
+	if (InputInterface)
+	{
+		InputInterface->SetForceFeedbackChannelValues(0, ForceFeedbackValues);
+	}
+}
+
+TStatId FPreviewForceFeedbackEffect::GetStatId() const
+{
+	RETURN_QUICK_DECLARE_CYCLE_STAT(FPreviewForceFeedbackEffect, STATGROUP_Tickables);
+}
+
+void FPreviewForceFeedbackEffect::AddReferencedObjects( FReferenceCollector& Collector )
+{
+	Collector.AddReferencedObject(ForceFeedbackEffect);
+}
+
+#undef LOCTEXT_NAMESPACE
diff --git a/Engine/Source/Developer/AssetTools/Private/AssetTypeActions/AssetTypeActions_ForceFeedbackEffect.h b/Engine/Source/Developer/AssetTools/Private/AssetTypeActions/AssetTypeActions_ForceFeedbackEffect.h
new file mode 100644
--- /dev/null
+++ b/Engine/Source/Developer/AssetTools/Private/AssetTypeActions/AssetTypeActions_ForceFeedbackEffect.h
@@ -0,0 +1,73 @@
+// Copyright Epic Games, Inc. All Rights Reserved.
+
+#pragma once
+
+#include "CoreMinimal.h"
+#include "Stats/Stats.h"
+#include "UObject/GCObject.h"
+#include "AssetTypeActions_Base.h"
+#include "GameFramework/ForceFeedbackEffect.h"
+#include "TickableEditorObject.h"
+
+struct FAssetData;
+
+struct FPreviewForceFeedbackEffect : public FActiveForceFeedbackEffect, public FTickableEditorObject, public FGCObject
+{
+	// FTickableEditorObject Implementation
+	virtual bool IsTickable() const override;
+	virtual void Tick( float DeltaTime ) override;
+	virtual TStatId GetStatId() const override;
+
+	// FGCObject Implementation
+	virtual void AddReferencedObjects( FReferenceCollector& Collector ) override;
+	virtual FString GetReferencerName() const override
+	{
+		return TEXT("FPreviewForceFeedbackEffect");
+	}
+};
+
+class FAssetTypeActions_ForceFeedbackEffect : public FAssetTypeActions_Base
+{
+public:
+	FAssetTypeActions_ForceFeedbackEffect(EAssetTypeCategories::Type InAssetCategoryBit)
+		: AssetCategoryBit(InAssetCategoryBit)
+	{ }
+
+	// IAssetTypeActions Implementation
+	virtual FText GetName() const override { return NSLOCTEXT("AssetTypeActions", "AssetTypeActions_ForceFeedbackEffect", "Force Feedback Effect"); }
+	virtual FColor GetTypeColor() const override { return FColor(175, 0, 0); }
+	virtual UClass* GetSupportedClass() const override;
+	virtual uint32 GetCategories() override { return AssetCategoryBit; }
+	virtual void GetActions(const TArray<UObject*>& InObjects, struct FToolMenuSection& Section) override;
+	virtual bool AssetsActivatedOverride( const TArray<UObject*>& InObjects, EAssetTypeActivationMethod::Type ActivationType ) override;
+	virtual TSharedPtr<SWidget> GetThumbnailOverlay(const FAssetData& AssetData) const override;
+
+	/** Return if the specified effect is playing*/
+	bool IsEffectPlaying(const TArray<TWeakObjectPtr<UForceFeedbackEffect>>& Objects) const;
+
+	/** Return if the specified effect is playing*/
+	bool IsEffectPlaying(const UForceFeedbackEffect* ForceFeedbackEffect) const;
+
+	/** Return if the specified asset is playing an effect*/
+	bool IsEffectPlaying(const FAssetData& AssetData) const;
+
+	/** Handler for when PlayEffect is selected */
+	void ExecutePlayEffect(TArray<TWeakObjectPtr<UForceFeedbackEffect>> Objects);
+
+	/** Handler for when StopEffect is selected */
+	void ExecuteStopEffect(TArray<TWeakObjectPtr<UForceFeedbackEffect>> Objects);
+
+	/** Returns true if only one effect is selected to play */
+	bool CanExecutePlayCommand(TArray<TWeakObjectPtr<UForceFeedbackEffect>> Objects) const;
+
+	/** Plays the specified effect */
+	void PlayEffect(UForceFeedbackEffect* Effect);
+
+	/** Stops any currently playing effect */
+	void StopEffect();
+
+private:
+	FPreviewForceFeedbackEffect PreviewForceFeedbackEffect;
+
+	EAssetTypeCategories::Type AssetCategoryBit;
+};
diff --git a/Engine/Source/Developer/AssetTools/Public/IAssetTypeActions.h b/Engine/Source/Developer/AssetTools/Public/IAssetTypeActions.h
--- a/Engine/Source/Developer/AssetTools/Public/IAssetTypeActions.h
+++ b/Engine/Source/Developer/AssetTools/Public/IAssetTypeActions.h
@@ -116,7 +116,6 @@ public:
 	virtual EThumbnailPrimType GetDefaultThumbnailPrimitiveType(UObject* Asset) const = 0;
 
 	/** Optionally returns a custom widget to overlay on top of this assets' thumbnail */
-	UE_DEPRECATED(5.5, "Convert the Action class over to the Definition system and use the new GetThumbnailActionOverlay, see UAssetDefinition for more information")
 	virtual TSharedPtr<class SWidget> GetThumbnailOverlay(const FAssetData& AssetData) const = 0;
 
 	/** Returns additional tooltip information for the specified asset, if it has any (otherwise return the null widget) */
diff --git a/Engine/Source/Editor/AssetDefinition/Public/AssetDefinition.h b/Engine/Source/Editor/AssetDefinition/Public/AssetDefinition.h
--- a/Engine/Source/Editor/AssetDefinition/Public/AssetDefinition.h
+++ b/Engine/Source/Editor/AssetDefinition/Public/AssetDefinition.h
@@ -325,16 +325,6 @@ public:
 };
 
 
-struct FAssetActionThumbnailOverlayInfo
-{
-public:
-	/** Image widget displayed in the top left border */
-	TSharedPtr<SWidget> ActionImageWidget;
-
-	/** Button widget displayed in the center when hovering the thumbnail */
-	TSharedPtr<SWidget> ActionButtonWidget;
-};
-
 /**
  * Asset Definitions represent top level assets that are known to the editor.
  *
@@ -575,19 +565,13 @@ public:
 	{
 		return nullptr;
 	}
-
-	UE_DEPRECATED(5.5, "Please use the new GetThumbnailActionOverlay function")
+	
 	/** Optionally returns a custom widget to overlay on top of this assets' thumbnail */
 	virtual TSharedPtr<SWidget> GetThumbnailOverlay(const FAssetData& InAssetData) const
 	{
 		return TSharedPtr<SWidget>();
 	}
 
-	virtual bool GetThumbnailActionOverlay(const FAssetData& InAssetData, FAssetActionThumbnailOverlayInfo& OutActionOverlayInfo) const
-	{
-		return false;
-	}
-
 	/** Whether this asset has external packages associated with it (impacts saving and dirty-state behavior) */
 	virtual bool ShouldSaveExternalPackages() const
 	{
diff --git a/Engine/Source/Editor/AudioEditor/Private/AssetTypeActions/AssetDefinition_SoundBase.cpp b/Engine/Source/Editor/AudioEditor/Private/AssetTypeActions/AssetDefinition_SoundBase.cpp
--- a/Engine/Source/Editor/AudioEditor/Private/AssetTypeActions/AssetDefinition_SoundBase.cpp
+++ b/Engine/Source/Editor/AudioEditor/Private/AssetTypeActions/AssetDefinition_SoundBase.cpp
@@ -195,57 +195,6 @@ TSharedPtr<SWidget> UAssetDefinition_SoundBase::GetThumbnailOverlay(const FAsset
 	return GetSoundBaseThumbnailOverlay(InAssetData, MoveTemp(OnClickedLambda));
 }
 
-bool UAssetDefinition_SoundBase::GetThumbnailActionOverlay(const FAssetData& InAssetData, FAssetActionThumbnailOverlayInfo& OutActionOverlayInfo) const
-{
-	auto OnGetDisplayBrushLambda = [InAssetData]() -> const FSlateBrush*
-	{
-		if (UE::AudioEditor::IsSoundPlaying(InAssetData))
-		{
-			return FAppStyle::GetBrush("ContentBrowser.AssetAction.StopIcon");
-		}
-
-		return FAppStyle::GetBrush("ContentBrowser.AssetAction.PlayIcon");
-	};
-
-	OutActionOverlayInfo.ActionImageWidget = SNew(SImage).Image_Lambda(OnGetDisplayBrushLambda);
-
-	auto OnToolTipTextLambda = [InAssetData]() -> FText
-	{
-		if (UE::AudioEditor::IsSoundPlaying(InAssetData))
-		{
-			return LOCTEXT("Thumbnail_StopSoundToolTip", "Stop selected sound");
-		}
-
-		return LOCTEXT("Thumbnail_PlaySoundToolTip", "Play selected sound");
-	};
-
-	auto OnClickedLambda = [InAssetData]() -> FReply
-	{
-		if (UE::AudioEditor::IsSoundPlaying(InAssetData))
-		{
-			UE::AudioEditor::StopSound();
-		}
-		else
-		{
-			// Load and play sound
-			UE::AudioEditor::PlaySound(Cast<USoundBase>(InAssetData.GetAsset()));
-		}
-		return FReply::Handled();
-	};
-
-	OutActionOverlayInfo.ActionButtonWidget = SNew(SButton)
-		.ButtonStyle(FAppStyle::Get(), "HoverHintOnly")
-		.ContentPadding(0.0f)
-		.ToolTipText_Lambda(OnToolTipTextLambda)
-		.OnClicked_Lambda(OnClickedLambda)
-		[
-			SNew(SImage)
-			.Image_Lambda(OnGetDisplayBrushLambda)
-		];
-
-	return true;
-}
-
 // Menu Extensions
 //--------------------------------------------------------------------
 
diff --git a/Engine/Source/Editor/AudioEditor/Private/AssetTypeActions/AssetDefinition_SoundSourceBus.h b/Engine/Source/Editor/AudioEditor/Private/AssetTypeActions/AssetDefinition_SoundSourceBus.h
--- a/Engine/Source/Editor/AudioEditor/Private/AssetTypeActions/AssetDefinition_SoundSourceBus.h
+++ b/Engine/Source/Editor/AudioEditor/Private/AssetTypeActions/AssetDefinition_SoundSourceBus.h
@@ -24,6 +24,5 @@ public:
 		return Categories;
 	}
 	virtual TSharedPtr<SWidget> GetThumbnailOverlay(const FAssetData& InAssetData) const override { return nullptr; }
-	virtual bool GetThumbnailActionOverlay(const FAssetData& InAssetData, FAssetActionThumbnailOverlayInfo& OutActionOverlayInfo) const override { return false; }
 	// UAssetDefinition End
 };
diff --git a/Engine/Source/Editor/AudioEditor/Private/AudioEditorModule.cpp b/Engine/Source/Editor/AudioEditor/Private/AudioEditorModule.cpp
--- a/Engine/Source/Editor/AudioEditor/Private/AudioEditorModule.cpp
+++ b/Engine/Source/Editor/AudioEditor/Private/AudioEditorModule.cpp
@@ -179,6 +179,7 @@ public:
 
 		AssetTools.RegisterAssetTypeActions(MakeShared<FAssetTypeActions_DialogueWave>());
 		AssetTools.RegisterAssetTypeActions(MakeShared<FAssetTypeActions_SoundClass>());
+		AssetTools.RegisterAssetTypeActions(MakeShared<FAssetTypeActions_SoundWave>());
 	}
 
 	virtual void RegisterAudioMixerAssetActions() override
diff --git a/Engine/Source/Editor/AudioEditor/Public/AssetTypeActions/AssetDefinition_SoundBase.h b/Engine/Source/Editor/AudioEditor/Public/AssetTypeActions/AssetDefinition_SoundBase.h
--- a/Engine/Source/Editor/AudioEditor/Public/AssetTypeActions/AssetDefinition_SoundBase.h
+++ b/Engine/Source/Editor/AudioEditor/Public/AssetTypeActions/AssetDefinition_SoundBase.h
@@ -35,7 +35,6 @@ public:
 		return Categories;
 	}
 	virtual TSharedPtr<SWidget> GetThumbnailOverlay(const FAssetData& InAssetData) const override;
-	virtual bool GetThumbnailActionOverlay(const FAssetData& InAssetData, FAssetActionThumbnailOverlayInfo& OutActionOverlayInfo) const override;
 	virtual EAssetCommandResult ActivateAssets(const FAssetActivateArgs& ActivateArgs) const override;
 	// UAssetDefinition End
 
diff --git a/Engine/Source/Editor/EditorStyle/Private/StarshipStyle.cpp b/Engine/Source/Editor/EditorStyle/Private/StarshipStyle.cpp
--- a/Engine/Source/Editor/EditorStyle/Private/StarshipStyle.cpp
+++ b/Engine/Source/Editor/EditorStyle/Private/StarshipStyle.cpp
@@ -6462,12 +6462,7 @@ void FStarshipEditorStyle::FStyle::SetupContentBrowserStyle()
 		// Misc
 		/** Should be moved, shared */ Set( "ContentBrowser.ThumbnailShadow", new BOX_BRUSH( "ContentBrowser/ThumbnailShadow" , FMargin( 4.0f / 64.0f ) ) );
 
-		// Playback Action
-		Set( "ContentBrowser.AssetAction.PlayIcon", new IMAGE_BRUSH_SVG("Starship/ContentBrowser/PlayIcon", Icon16x16));
-		Set( "ContentBrowser.AssetAction.StopIcon", new IMAGE_BRUSH_SVG("Starship/ContentBrowser/StopIcon", Icon16x16));
-		FLinearColor OverlayColorAssetThumbnailOverlay = FStyleColors::Panel.GetSpecifiedColor();
-		OverlayColorAssetThumbnailOverlay.A = 0.75f;
-		Set("ContentBrowser.AssetTileItem.AssetThumbnailBar", new FSlateRoundedBoxBrush(OverlayColorAssetThumbnailOverlay, 4.f));
+
 
 		Set( "ContentBrowser.ColumnViewAssetIcon", new IMAGE_BRUSH( "Icons/doc_16x", Icon16x16 ) );
 
diff --git a/Engine/Source/Editor/UnrealEd/Private/AssetThumbnail.cpp b/Engine/Source/Editor/UnrealEd/Private/AssetThumbnail.cpp
--- a/Engine/Source/Editor/UnrealEd/Private/AssetThumbnail.cpp
+++ b/Engine/Source/Editor/UnrealEd/Private/AssetThumbnail.cpp
@@ -33,7 +33,6 @@
 #include "AssetRegistry/IAssetRegistry.h"
 #include "ShaderCompiler.h"
 #include "AssetCompilingManager.h"
-#include "AssetDefinitionRegistry.h"
 #include "IAssetTools.h"
 #include "AssetTypeActions_Base.h"
 #include "AssetToolsModule.h"
@@ -331,75 +330,14 @@ public:
 
 		if( InArgs._AllowAssetSpecificThumbnailOverlay && AssetTypeActions.IsValid() )
 		{
-#if UE_CONTENTBROWSER_NEW_STYLE
-			const UAssetDefinition* AssetDefinition = UAssetDefinitionRegistry::Get()->GetAssetDefinitionForClass(Class);
-
-			FAssetActionThumbnailOverlayInfo OutThumbnailInfo;
-			if (AssetDefinition && AssetDefinition->GetThumbnailActionOverlay(AssetData, OutThumbnailInfo))
-			{
-				constexpr int32 OverlayZOrder = 1;
-				if (OutThumbnailInfo.ActionImageWidget.IsValid())
-				{
-					constexpr float PaddingFromTopLeftBorder = 2.f;
-					constexpr float TopLeftImageSize = 20.f;
-
-					OverlayWidget->AddSlot()
-					.ZOrder(OverlayZOrder)
-					.VAlign(VAlign_Top)
-					.HAlign(HAlign_Left)
-					.Padding(PaddingFromTopLeftBorder, PaddingFromTopLeftBorder, 0.f, 0.f)
-					[
-						SNew(SBorder)
-						.BorderImage(FAppStyle::GetBrush(TEXT("ContentBrowser.AssetTileItem.AssetThumbnailBar")))
-						.Visibility_Lambda([this] () { return IsHovered() ? EVisibility::Collapsed : EVisibility::Visible; })
-						[
-							SNew(SBox)
-							.WidthOverride(TopLeftImageSize)
-							.HeightOverride(TopLeftImageSize)
-							[
-								OutThumbnailInfo.ActionImageWidget.ToSharedRef()
-							]
-						]
-					];
-				}
-
-				if (OutThumbnailInfo.ActionButtonWidget.IsValid())
-				{
-					constexpr float CenterImageSize = 32.f;
-
-					OverlayWidget->AddSlot()
-					.ZOrder(OverlayZOrder)
-					.VAlign(VAlign_Center)
-					.HAlign(HAlign_Center)
-					[
-						SNew(SBorder)
-						.BorderImage(FAppStyle::GetBrush(TEXT("ContentBrowser.AssetTileItem.AssetThumbnailBar")))
-						.Visibility_Lambda([this] () { return IsHovered() ? EVisibility::Visible : EVisibility::Collapsed; })
-						[
-							SNew(SBox)
-							.WidthOverride(CenterImageSize)
-							.HeightOverride(CenterImageSize)
-							[
-								OutThumbnailInfo.ActionButtonWidget.ToSharedRef()
-							]
-						]
-					];
-				}
-			}
-			else
-#endif
+			// Does the asset provide an additional thumbnail overlay?
+			TSharedPtr<SWidget> AssetSpecificThumbnailOverlay = AssetTypeActions->GetThumbnailOverlay(AssetData);
+			if( AssetSpecificThumbnailOverlay.IsValid() )
 			{
-				// Does the asset provide an additional thumbnail overlay?
-				PRAGMA_DISABLE_DEPRECATION_WARNINGS
-				TSharedPtr<SWidget> AssetSpecificThumbnailOverlay = AssetTypeActions->GetThumbnailOverlay(AssetData);
-				PRAGMA_ENABLE_DEPRECATION_WARNINGS
-				if( AssetSpecificThumbnailOverlay.IsValid() )
-				{
-					OverlayWidget->AddSlot()
-					[
-						AssetSpecificThumbnailOverlay.ToSharedRef()
-					];
-				}
+				OverlayWidget->AddSlot()
+				[
+					AssetSpecificThumbnailOverlay.ToSharedRef()
+				];
 			}
 		}
 
