Fix crash in UAssetToolsImpl::CreateAssetWithDialog where the UFactory::GetSupportedClass() is deferenced without check if its null.

#rnx

diff --git a/Engine/Source/Developer/AssetTools/Private/AssetTools.cpp b/Engine/Source/Developer/AssetTools/Private/AssetTools.cpp
--- a/Engine/Source/Developer/AssetTools/Private/AssetTools.cpp
+++ b/Engine/Source/Developer/AssetTools/Private/AssetTools.cpp
@@ -1848,7 +1848,20 @@ UObject* UAssetToolsImpl::CreateAssetWithDialog(const FString& AssetName, const
 		SaveAssetDialogConfig.DefaultPath = PackagePath;
 		SaveAssetDialogConfig.DefaultAssetName = AssetName;
 		SaveAssetDialogConfig.ExistingAssetPolicy = ESaveAssetDialogExistingAssetPolicy::AllowButWarn;
-		SaveAssetDialogConfig.AssetClassNames.Add(Factory->GetSupportedClass()->GetClassPathName());
+		if (Factory->GetSupportedClass() != nullptr)
+		{
+			SaveAssetDialogConfig.AssetClassNames.Add(Factory->GetSupportedClass()->GetClassPathName());
+		}
+		else if(AssetClass && Factory->DoesSupportClass(AssetClass))
+		{
+			SaveAssetDialogConfig.AssetClassNames.Add(AssetClass->GetClassPathName());
+		}
+		else
+		{
+			FMessageDialog::Open(EAppMsgType::Ok, LOCTEXT("InvalidFactory", "Cannot create the new asset because the supplied factory does not support the supplied class."));
+			return nullptr;
+		}
+		
 
 		FContentBrowserModule& ContentBrowserModule = FModuleManager::LoadModuleChecked<FContentBrowserModule>("ContentBrowser");
 		FString SaveObjectPath = ContentBrowserModule.Get().CreateModalSaveAssetDialog(SaveAssetDialogConfig);
