Fix localisation error
#rnx

diff --git a/Engine/Source/Developer/AssetTools/Private/AssetTools.cpp b/Engine/Source/Developer/AssetTools/Private/AssetTools.cpp
--- a/Engine/Source/Developer/AssetTools/Private/AssetTools.cpp
+++ b/Engine/Source/Developer/AssetTools/Private/AssetTools.cpp
@@ -122,6 +122,14 @@
 
 #define LOCTEXT_NAMESPACE "AssetTools"
 
+namespace UE::Private
+{
+	FText GetInvalidFactoryMessage()
+	{
+		return LOCTEXT("InvalidFactory", "Cannot create the new asset because the supplied factory does not support the supplied class.");
+	} 
+}
+
 class FAssetDefinitionProxy : public FAssetTypeActions_Base
 {
 public:
@@ -1721,7 +1729,7 @@ UObject* UAssetToolsImpl::CreateAsset(const FString& AssetName, const FString& P
 		((Factory->SupportedClass != nullptr && !ensure(AssetClass->IsChildOf(Factory->GetSupportedClass()))) || 
 		 (Factory->SupportedClass == nullptr && !ensure(Factory->DoesSupportClass(AssetClass)))) )
 	{
-		FMessageDialog::Open( EAppMsgType::Ok, LOCTEXT("InvalidFactory", "The new asset wasn't created because the supplied factory does not support the supplied class.") );
+		FMessageDialog::Open( EAppMsgType::Ok, UE::Private::GetInvalidFactoryMessage());
 		return nullptr;
 	}
 
@@ -1858,7 +1866,7 @@ UObject* UAssetToolsImpl::CreateAssetWithDialog(const FString& AssetName, const
 		}
 		else
 		{
-			FMessageDialog::Open(EAppMsgType::Ok, LOCTEXT("InvalidFactory", "Cannot create the new asset because the supplied factory does not support the supplied class."));
+			FMessageDialog::Open(EAppMsgType::Ok, UE::Private::GetInvalidFactoryMessage());
 			return nullptr;
 		}
 		
