New localization option to move localized variants alongside source asset in the content browser and another one to detect if those localized variants are only in Revision Control:
 - Project Settings / Editor / Assets / Internationalization / Rename Localized Variants Alongside Source Asset (defaults to: true)
 - Project Settings / Editor / Revision Control / Internationalization / Requires Revision Control to Rename Localizable Assets (defaults to: false) (works only with Perforce)

[RN]

diff --git a/Engine/Plugins/Developer/PerforceSourceControl/Source/PerforceSourceControl/Private/PerforceSourceControlOperations.cpp b/Engine/Plugins/Developer/PerforceSourceControl/Source/PerforceSourceControl/Private/PerforceSourceControlOperations.cpp
--- a/Engine/Plugins/Developer/PerforceSourceControl/Source/PerforceSourceControl/Private/PerforceSourceControlOperations.cpp
+++ b/Engine/Plugins/Developer/PerforceSourceControl/Source/PerforceSourceControl/Private/PerforceSourceControlOperations.cpp
@@ -268,16 +268,17 @@ static void ParseRecordSet(const FP4RecordSet& InRecords, TArray<FText>& OutResu
 }
 
 /** Simple parsing of a record set to update state */
-static void ParseRecordSetForState(const FP4RecordSet& InRecords, TMap<FString, EPerforceState::Type>& OutResults)
+static void ParseRecordSetForState(const FP4RecordSet& InRecords, TMap<FString, EPerforceState::Type>& OutResults, bool bConsiderDepotFile = false)
 {
 	// Iterate over each record found as a result of the command, parsing it for relevant information
 	for (const FP4Record& ClientRecord : InRecords)
 	{
-		const FString& FileName = ClientRecord(TEXT("clientFile"));
+		const FString& ClientFileName = ClientRecord(TEXT("clientFile"));
+		const FString& DepotFileName = ClientRecord(TEXT("depotFile"));
 		const FString& Action = ClientRecord(TEXT("action"));
 
-		check(FileName.Len());
-		FString FullPath(FileName);
+		check(ClientFileName.Len() || bConsiderDepotFile ? DepotFileName.Len() : false);
+		FString FullPath = ClientFileName.Len() ? ClientFileName : (bConsiderDepotFile ? DepotFileName : FString());
 		FPaths::NormalizeFilename(FullPath);
 
 		if(Action.Len() > 0)
@@ -805,10 +806,14 @@ bool FPerforceGetFileListWorker::Execute(FPerforceSourceControlCommand& InComman
 			Parameters.Add(TEXT("-e"));
 		}
 		Parameters.Append(InCommand.Files);
-		AppendMaskParameter(Parameters);
+		if (Operation->GetMethodUsed() == FGetFileList::FolderSearch)
+		{
+			AppendMaskParameter(Parameters);
+		}
 		FP4RecordSet Records;
 		InCommand.bCommandSuccessful = Connection.RunCommand(TEXT("files"), Parameters, Records, InCommand.ResultInfo, FOnIsCancelled::CreateRaw(&InCommand, &FPerforceSourceControlCommand::IsCanceled), InCommand.bConnectionDropped);
-		ParseRecordSetForState(Records, OutResults);
+		bool bConsiderDepotFiles = Operation->GetMethodUsed() == FGetFileList::FileRegexSearch;
+		ParseRecordSetForState(Records, OutResults, bConsiderDepotFiles);
 
 		TArray<FString> FilesList;
 		FilesList.Reserve(OutResults.Num());
diff --git a/Engine/Plugins/Editor/ContentBrowser/ContentBrowserAssetDataSource/Source/ContentBrowserAssetDataSource/Private/ContentBrowserAssetDataCore.cpp b/Engine/Plugins/Editor/ContentBrowser/ContentBrowserAssetDataSource/Source/ContentBrowserAssetDataSource/Private/ContentBrowserAssetDataCore.cpp
--- a/Engine/Plugins/Editor/ContentBrowser/ContentBrowserAssetDataSource/Source/ContentBrowserAssetDataSource/Private/ContentBrowserAssetDataCore.cpp
+++ b/Engine/Plugins/Editor/ContentBrowser/ContentBrowserAssetDataSource/Source/ContentBrowserAssetDataSource/Private/ContentBrowserAssetDataCore.cpp
@@ -1099,7 +1099,9 @@ bool RenameAssetFileItem(IAssetTools* InAssetTools, const FContentBrowserAssetFi
 		const FString PackagePath = FPackageName::GetLongPackagePath(Asset->GetOutermost()->GetName());
 
 		TArray<FAssetRenameData> AssetsAndNames;
-		AssetsAndNames.Emplace(FAssetRenameData(Asset, PackagePath, InNewName));
+		const bool bSoftReferenceOnly = false;
+		const bool bAlsoRenameLocalizedVariants = true;
+		AssetsAndNames.Emplace(FAssetRenameData(Asset, PackagePath, InNewName, bSoftReferenceOnly, bAlsoRenameLocalizedVariants));
 		// Note: This also returns false for Pending results as the rename may yet fail or be canceled, so the change has to be detected later via the asset registry
 		return InAssetTools->RenameAssetsWithDialog(AssetsAndNames) == EAssetRenameResult::Success;
 	}
diff --git a/Engine/Source/Developer/AssetTools/Private/AssetRenameManager.cpp b/Engine/Source/Developer/AssetTools/Private/AssetRenameManager.cpp
--- a/Engine/Source/Developer/AssetTools/Private/AssetRenameManager.cpp
+++ b/Engine/Source/Developer/AssetTools/Private/AssetRenameManager.cpp
@@ -1,7 +1,7 @@
 // Copyright Epic Games, Inc. All Rights Reserved.
 
-
 #include "AssetRenameManager.h"
+#include "AssetDefinitionRegistry.h"
 #include "Serialization/ArchiveUObject.h"
 #include "UObject/Class.h"
 #include "Misc/PackageName.h"
@@ -46,6 +46,7 @@
 #include "Kismet2/BlueprintEditorUtils.h"
 #include "Misc/RedirectCollector.h"
 #include "Settings/BlueprintEditorProjectSettings.h"
+#include "Settings/EditorProjectSettings.h"
 #include "AssetToolsLog.h"
 #include "Engine/Level.h"
 #include "Engine/World.h"
@@ -53,6 +54,10 @@
 #include "GameMapsSettings.h"
 #include "AssetToolsModule.h"
 #include "IAssetTools.h"
+#include "Internationalization/PackageLocalizationUtil.h"
+#include "IAssetTools.h"
+#include "ILocalizedAssetTools.h"
+#include "SFileListReportDialog.h"
 
 #define LOCTEXT_NAMESPACE "AssetRenameManager"
 
@@ -124,140 +129,29 @@ struct FAssetRenameDataWithReferencers : public FAssetRenameData
 	}
 };
 
-class SRenameFailures : public SCompoundWidget
-{
-public:
-	SLATE_BEGIN_ARGS(SRenameFailures){}
-
-		SLATE_ARGUMENT(TArray<FText>, FailedRenames)
-
-	SLATE_END_ARGS()
-
-	BEGIN_SLATE_FUNCTION_BUILD_OPTIMIZATION
-	void Construct( const FArguments& InArgs )
-	{
-		for (const FText& RenameText : InArgs._FailedRenames)
-		{
-			FailedRenames.Add( MakeShareable( new FText(RenameText) ) );
-		}
-
-		ChildSlot
-		[
-			SNew(SBorder)
-			.BorderImage( FAppStyle::GetBrush("Docking.Tab.ContentAreaBrush") )
-			.Padding(FMargin(4, 8, 4, 4))
-			[
-				SNew(SVerticalBox)
-
-				// Title text
-				+SVerticalBox::Slot()
-				.AutoHeight()
-				[
-					SNew(STextBlock) .Text( LOCTEXT("RenameFailureTitle", "The following assets could not be renamed.") )
-				]
-
-				// Failure list
-				+SVerticalBox::Slot()
-				.Padding(0, 8)
-				.FillHeight(1.f)
-				[
-					SNew(SBorder)
-					.BorderImage( FAppStyle::GetBrush("ToolPanel.GroupBorder") )
-					[
-						SNew(SListView<TSharedRef<FText>>)
-						.ListItemsSource(&FailedRenames)
-						.SelectionMode(ESelectionMode::None)
-						.OnGenerateRow(this, &SRenameFailures::MakeListViewWidget)
-					]
-				]
-
-				// Close button
-				+SVerticalBox::Slot()
-				.AutoHeight()
-				.Padding(0, 4)
-				.HAlign(HAlign_Right)
-				[
-					SNew(SButton)
-					.OnClicked(this, &SRenameFailures::CloseClicked)
-					.Text(LOCTEXT("RenameFailuresCloseButton", "Close"))
-				]
-			]
-		];
-	}
-	END_SLATE_FUNCTION_BUILD_OPTIMIZATION
-
-	static void OpenRenameFailuresDialog(const TArray<FText>& InFailedRenames)
-	{
-		TSharedRef<SWindow> RenameWindow = SNew(SWindow)
-			.Title(LOCTEXT("FailedRenamesDialog", "Failed Renames"))
-			.ClientSize(FVector2D(800,400))
-			.SupportsMaximize(false)
-			.SupportsMinimize(false)
-			[
-				SNew(SRenameFailures).FailedRenames(InFailedRenames)
-			];
-
-		IMainFrameModule& MainFrameModule = FModuleManager::LoadModuleChecked<IMainFrameModule>(TEXT("MainFrame"));
-
-		if (MainFrameModule.GetParentWindow().IsValid())
-		{
-			FSlateApplication::Get().AddWindowAsNativeChild(RenameWindow, MainFrameModule.GetParentWindow().ToSharedRef());
-		}
-		else
-		{
-			FSlateApplication::Get().AddWindow(RenameWindow);
-		}
-	}
-
-private:
-	TSharedRef<ITableRow> MakeListViewWidget(TSharedRef<FText> Item, const TSharedRef<STableViewBase>& OwnerTable)
-	{
-		return
-			SNew(STableRow< TSharedRef<FText> >, OwnerTable)
-			[
-				SNew(STextBlock).Text(Item.Get())
-			];
-	}
-
-	FReply CloseClicked()
-	{
-		TSharedPtr<SWindow> Window = FSlateApplication::Get().FindWidgetWindow(AsShared());
-
-		if (Window.IsValid())
-		{
-			Window->RequestDestroyWindow();
-		}
-
-		return FReply::Handled();
-	}
-
-private:
-	TArray< TSharedRef<FText> > FailedRenames;
-};
-
-
 ///////////////////////////
 // FAssetRenameManager
 ///////////////////////////
 
-/** Renames assets using the specified names. */
 bool FAssetRenameManager::RenameAssets(const TArray<FAssetRenameData>& AssetsAndNames) const
 {
-	// If the asset registry is still loading assets, we cant check for referencers, so we must open the rename dialog
+	const bool bAutoCheckout = true;
+	const bool bWithDialog = false;
+
+	// If the asset registry is still loading assets, we cant check for referencers and we cant open dialogs so we fail
 	FAssetRegistryModule& AssetRegistryModule = FModuleManager::LoadModuleChecked<FAssetRegistryModule>("AssetRegistry");
 	if (AssetRegistryModule.Get().IsLoadingAssets())
 	{
 		UE_LOG(LogAssetTools, Warning, TEXT("Unable To Rename While Discovering Assets"));
 		return false;
 	}
-	const bool bAutoCheckout = true;
-	const bool bWithDialog = false;
-	return FixReferencesAndRename(AssetsAndNames, bAutoCheckout, bWithDialog);
+
+	return RenameAssetsAndVariants(AssetsAndNames, bAutoCheckout, bWithDialog);
 }
 
 EAssetRenameResult FAssetRenameManager::RenameAssetsWithDialog(const TArray<FAssetRenameData>& AssetsAndNames, bool bAutoCheckout) const
 {
-	bool bWithDialog = true;
+	const bool bWithDialog = true;
 
 	// If the asset registry is still loading assets, we cant check for referencers, so we must open the rename dialog
 	FAssetRegistryModule& AssetRegistryModule = FModuleManager::LoadModuleChecked<FAssetRegistryModule>("AssetRegistry");
@@ -265,17 +159,15 @@ EAssetRenameResult FAssetRenameManager::RenameAssetsWithDialog(const TArray<FAss
 	{
 		// Open a dialog asking the user to wait while assets are being discovered
 		SDiscoveringAssetsDialog::OpenDiscoveringAssetsDialog(
-			SDiscoveringAssetsDialog::FOnAssetsDiscovered::CreateSP(this, &FAssetRenameManager::FixReferencesAndRenameCallback, AssetsAndNames, bAutoCheckout, bWithDialog)
+			SDiscoveringAssetsDialog::FOnAssetsDiscovered::CreateSP(this, &FAssetRenameManager::RenameAssetsAndVariantsCallback, AssetsAndNames, bAutoCheckout, bWithDialog)
 		);
 		return EAssetRenameResult::Pending;
 	}
-	else
-	{
-		// No need to wait, attempt to fix references and rename now.
-		return FixReferencesAndRename(AssetsAndNames, bAutoCheckout, bWithDialog)
-			? EAssetRenameResult::Success
-			: EAssetRenameResult::Failure;
-	}
+
+	// No need to wait, check for variants and attempt to fix references now.
+	return RenameAssetsAndVariants(AssetsAndNames, bAutoCheckout, bWithDialog)
+		? EAssetRenameResult::Success
+		: EAssetRenameResult::Failure;
 }
 
 void FAssetRenameManager::FindSoftReferencesToObject(FSoftObjectPath TargetObject, TArray<UObject*>& ReferencingObjects) const
@@ -320,13 +212,147 @@ void FAssetRenameManager::FindSoftReferencesToObjects(const TArray<FSoftObjectPa
 	GatherReferencingObjects(AssetsToRename, ReferencingObjects);
 }
 
-void FAssetRenameManager::FixReferencesAndRenameCallback(TArray<FAssetRenameData> AssetsAndNames, bool bAutoCheckout, bool bWithDialog) const
+void FAssetRenameManager::RenameAssetsAndVariantsCallback(TArray<FAssetRenameData> InAssetsRenameData, bool bAutoCheckout, bool bWithDialog) const
 {
-	FixReferencesAndRename(AssetsAndNames, bAutoCheckout, bWithDialog);
+	RenameAssetsAndVariants(InAssetsRenameData, bAutoCheckout, bWithDialog);
+}
+
+bool FAssetRenameManager::RenameAssetsAndVariants(const TArray<FAssetRenameData>& InAssetsRenameData, bool bAutoCheckout, bool bWithDialog) const
+{
+	FScopedSlowTask RenamingSlowTask(11.0f, LOCTEXT("RenamingSlowTask", "Renaming assets..."), bWithDialog);
+	bool bShowCancelButton = false;
+	bool bAllowInPIE = true;
+	RenamingSlowTask.MakeDialog(bShowCancelButton, bAllowInPIE);
+	RenamingSlowTask.EnterProgressFrame(2.0f, LOCTEXT("CheckingForLocalizedVariantsSlowTask", "Checking for localized variants... This might take a long time."));
+	RenamingSlowTask.ForceRefresh();
+
+	TArray<FAssetRenameData> AssetsAndVariants;
+
+	TArray<FAssetRenameData> AssetsToCheckForVariants;
+	if (RequiresCheckingForVariants(InAssetsRenameData, AssetsAndVariants, AssetsToCheckForVariants))
+	{
+		// Build a list of packages to check for variants (and build the map to remap the rename data later)
+		TArray<FName> SourcePackagesToCheckForVariants;
+		TMap<FName, const FAssetRenameData*> MapSourcePackageToRenameData;
+		SourcePackagesToCheckForVariants.Reserve(AssetsToCheckForVariants.Num());
+		MapSourcePackageToRenameData.Reserve(AssetsToCheckForVariants.Num());
+		for (const FAssetRenameData& AssetToCheckForVariants : AssetsToCheckForVariants)
+		{
+			FString SourcePackageString;
+			FPackageLocalizationUtil::ConvertToSource(AssetToCheckForVariants.Asset->GetOuter()->GetName(), SourcePackageString);
+			FName SourcePackage(SourcePackageString);
+			SourcePackagesToCheckForVariants.Add(SourcePackage);
+			MapSourcePackageToRenameData.Add(SourcePackage, &AssetToCheckForVariants);
+		}
+
+		IAssetTools& AssetTools = FModuleManager::LoadModuleChecked<FAssetToolsModule>("AssetTools").Get();
+		TSharedPtr<ILocalizedAssetTools> LocalizedAssetTools = AssetTools.GetLocalizedAssetTools();
+
+		TArray<FName> SourcePackagesWithoutVariants;
+		TMap<FName, TArray<FName>> VariantsBySourcesOnDisk;
+		TMap<FName, TArray<FName>> VariantsBySourcesOnlyInRevisionControl;
+		bool bNeedToCheckInRevisionControl = USourceControlPreferences::RequiresRevisionControlToRenameLocalizableAssets();
+		ELocalizedAssetsResult Result = LocalizedAssetTools->GetLocalizedVariants(SourcePackagesToCheckForVariants, VariantsBySourcesOnDisk, bNeedToCheckInRevisionControl, VariantsBySourcesOnlyInRevisionControl, &SourcePackagesWithoutVariants);
+		bool bRevisionControlWasNeeded = (Result == ELocalizedAssetsResult::RevisionControlNotAvailable);
+
+		// If all variants are accessible (not in Revision Control), we can now
+		// rebuild all the renaming data from the localized variants information
+		if (!bRevisionControlWasNeeded && VariantsBySourcesOnlyInRevisionControl.IsEmpty())
+		{
+			IAssetRegistry& AssetRegistry = IAssetRegistry::GetChecked();
+
+			// Let's process each "VariantsBySourcesOnDisk" one-by-one and map it back to its rename data
+			for (const TPair<FName, TArray<FName>>& VariantsBySourceOnDisk : VariantsBySourcesOnDisk)
+			{
+				// Find the rename data that match this source
+				const FAssetRenameData& CurrentRenameData = *MapSourcePackageToRenameData[VariantsBySourceOnDisk.Key];
+
+				// Get some rename data information (at this point, we don't know if it is a Source Asset or a Localized Variant rename data information)
+				FSoftObjectPath OldAssetSoftObjectPath(CurrentRenameData.Asset.Get());
+				FString OldAssetPackageName = OldAssetSoftObjectPath.GetLongPackageName();
+				FString OldAssetPath = OldAssetSoftObjectPath.GetAssetPath().ToString();
+
+				// Create and add the Source Asset Rename Data
+				FString SourceAssetOldPath;
+				FPackageLocalizationUtil::ConvertToSource(OldAssetPath, SourceAssetOldPath);
+				UObject* SourceObjectPtr = AssetRegistry.GetAssetByObjectPath(SourceAssetOldPath).GetAsset();
+				FString SourceNewPackagePath;
+				FPackageLocalizationUtil::ConvertToSource(CurrentRenameData.NewPackagePath, SourceNewPackagePath);
+				FAssetRenameData CurrentAssetRenameData(SourceObjectPtr,
+					SourceNewPackagePath,
+					CurrentRenameData.NewName,
+					CurrentRenameData.bOnlyFixSoftReferences,
+					CurrentRenameData.bAlsoRenameLocalizedVariants);
+				AssetsAndVariants.Add(CurrentAssetRenameData);
+
+				// Create and add all Variants Asset Rename Data
+				for (const FName& Variant : VariantsBySourceOnDisk.Value)
+				{
+					FString Culture;
+					FPackageLocalizationUtil::ExtractCultureFromLocalized(Variant.ToString(), Culture);
+
+					FString OldLocalizedAssetPath;
+					FString NewLocalizedPackagePath;
+					FPackageLocalizationUtil::ConvertSourceToLocalized(SourceAssetOldPath, Culture, OldLocalizedAssetPath);
+					FPackageLocalizationUtil::ConvertSourceToLocalized(SourceNewPackagePath, Culture, NewLocalizedPackagePath);
+
+					UObject* LocalizedObjectPtr = AssetRegistry.GetAssetByObjectPath(OldLocalizedAssetPath).GetAsset();
+
+					// Create the renaming data for the localized variant then add it to the list of files to be renamed
+					CurrentAssetRenameData.Asset = LocalizedObjectPtr;
+					CurrentAssetRenameData.NewPackagePath = NewLocalizedPackagePath;
+					AssetsAndVariants.Add(CurrentAssetRenameData);
+				}
+			}
+
+			// Finally, let's add back all the assets that don't have any variants as-is
+			for (const FName& SourcePackageWithoutVariants : SourcePackagesWithoutVariants)
+			{
+				AssetsAndVariants.Add(*MapSourcePackageToRenameData[SourcePackageWithoutVariants]);
+			}
+		}
+
+		// Process error and interrupt the renaming process if necessary
+		if (bRevisionControlWasNeeded)
+		{
+			// Revision Control is not available. Renaming must fail.
+			UE_LOG(LogAssetTools, Warning, TEXT("%s"), *LocalizedAssetTools->GetRevisionControlIsNotAvailableWarningText().ToString());
+			RenameInterrupted(InAssetsRenameData, LocalizedAssetTools->GetRevisionControlIsNotAvailableWarningText(), bWithDialog);
+			if (bWithDialog)
+			{
+				LocalizedAssetTools->OpenRevisionControlRequiredDialog();
+			}
+			return false;
+		}
+		else if (!VariantsBySourcesOnlyInRevisionControl.IsEmpty())
+		{
+			TArray<FText> LocalizedVariantsPathsRequiredToSync;
+			for (const TPair<FName, TArray<FName>>& VariantsBySourceOnlyInRevisionControl : VariantsBySourcesOnlyInRevisionControl)
+			{
+				for (const FName& VariantInRevisionControl : VariantsBySourceOnlyInRevisionControl.Value)
+				{
+					LocalizedVariantsPathsRequiredToSync.Add(FText::AsCultureInvariant(VariantInRevisionControl.ToString()));
+				}
+			}
+
+			UE_LOG(LogAssetTools, Warning, TEXT("A file that needs to be renamed was detected in Revision Control but not on your disk."));
+			RenameInterrupted(InAssetsRenameData, LocalizedAssetTools->GetFilesNeedToBeOnDiskWarningText(), bWithDialog);
+			if (bWithDialog)
+			{
+				LocalizedAssetTools->OpenFilesInRevisionControlRequiredDialog(LocalizedVariantsPathsRequiredToSync);
+			}
+			return false;
+		}
+	}
+
+	// We now have a full list of Assets to rename and their variants if applicable. Let's continue the renaming process.
+	return FixReferencesAndRename(AssetsAndVariants, bAutoCheckout, bWithDialog, RenamingSlowTask);
 }
 
-bool FAssetRenameManager::FixReferencesAndRename(const TArray<FAssetRenameData>& AssetsAndNames, bool bAutoCheckout, bool bWithDialog) const
+bool FAssetRenameManager::FixReferencesAndRename(const TArray<FAssetRenameData>& AssetsAndNames, bool bAutoCheckout, bool bWithDialog, FScopedSlowTask& RenamingSlowTask) const
 {
+	RenamingSlowTask.EnterProgressFrame(0.9f, LOCTEXT("FixingReferencesSlowTask", "Finding references..."));
+
 	bool bSoftReferencesOnly = true;
 	// Prep a list of assets to rename with an extra boolean to determine if they should leave a redirector or not
 	TArray<FAssetRenameDataWithReferencers> AssetsToRename;
@@ -422,6 +448,7 @@ bool FAssetRenameManager::FixReferencesAndRename(const TArray<FAssetRenameData>&
 		LoadReferencingPackages(AssetsToRename, bSoftReferencesOnly, true, ReferencingPackagesToSave, SoftReferencingObjects);
 
 		// Prompt to check out source package and all referencing packages, leave redirectors for assets referenced by packages that are not checked out and remove those packages from the save list.
+		RenamingSlowTask.EnterProgressFrame(0.1f, LOCTEXT("CheckOutPackagesSlowTask", "Check out packages..."));
 		const bool bUserAcceptedCheckout = CheckOutPackages(AssetsToRename, ReferencingPackagesToSave, bAutoCheckout);
 
 		if (bUserAcceptedCheckout || bSoftReferencesOnly)
@@ -440,7 +467,7 @@ bool FAssetRenameManager::FixReferencesAndRename(const TArray<FAssetRenameData>&
 				if (ReferencingPackagesToSave.Num() > 0)
 				{
 					// Only do the rename if there are actually packages with references
-					PerformAssetRename(AssetsToRename);
+					PerformAssetRename(AssetsToRename, RenamingSlowTask);
 
 					for (const FAssetRenameDataWithReferencers& RenameData : AssetsToRename)
 					{
@@ -476,7 +503,7 @@ bool FAssetRenameManager::FixReferencesAndRename(const TArray<FAssetRenameData>&
 			{
 				// Perform the rename, leaving redirectors only for assets which need them
 				// Also save all packages that were referencing any of the assets that were moved without redirectors
-				PerformAssetRename(AssetsToRename, ReferencingPackagesToSave);
+				PerformAssetRename(AssetsToRename, ReferencingPackagesToSave, RenamingSlowTask);
 
 				// Issue post rename event
 				AssetPostRenameEvent.Broadcast(AssetsAndNames);
@@ -709,7 +736,7 @@ void FAssetRenameManager::FindCDOReferences(const TArrayView<FAssetRenameDataWit
 
 				// Resolve to the redirected asset path if necessary
 				FSoftObjectPath FinalSoftObjPath = SoftRefObjPath.GetWithoutSubPath();
-			
+
 				if (!Object.IsValid() && SoftRefObjPath.IsValid())
 				{
 					FSoftObjectPath RedirObjectPath = GRedirectCollector.GetAssetPathRedirection(SoftRefObjPath.GetWithoutSubPath());
@@ -884,7 +911,7 @@ void FAssetRenameManager::LoadReferencingPackages(TArray<FAssetRenameDataWithRef
 		{
 			GWarn->StatusUpdate(AssetIdx, AssetsToRename.Num(), ReferenceUpdateSlowTask);
 		}
-		
+
 		FAssetRenameDataWithReferencers& RenameData = AssetsToRename[AssetIdx];
 
 		TSet<FName> ReferencingExternalPackageNames;
@@ -970,7 +997,7 @@ void FAssetRenameManager::LoadReferencingPackages(TArray<FAssetRenameDataWithRef
 			// Don't load package if this is a soft reference fix and the project settings say not to
 			if (!Package && (!RenameData.bOnlyFixSoftReferences || bLoadPackagesForSoftReferences))
 			{
-				if(!bStartedSlowTask)
+				if (!bStartedSlowTask)
 				{
 					bStartedSlowTask = true;
 					GWarn->BeginSlowTask(ReferenceUpdateSlowTask, true);
@@ -979,7 +1006,7 @@ void FAssetRenameManager::LoadReferencingPackages(TArray<FAssetRenameDataWithRef
 			}
 
 			if (Package)
-			{				
+			{
 				bool bFoundSoftReference = CheckPackageForSoftObjectReferences(Package, ModifiedPaths, OutSoftReferencingObjects);
 
 				// Only add to list if we're doing a hard reference fixup or we found a soft reference
@@ -1278,7 +1305,7 @@ void FAssetRenameManager::SetupPublicAssets(TArray<FAssetRenameDataWithReference
 	TArray<UPackage*> PackagesToMakePublic;
 	TSet<FName> RenamedReferencedPackageNames;
 	bool bSomeAssetCannotBeMovedToAnotherPlugin = false;
-	
+
 	// Determine which moved asset need to become public
 	for (FAssetRenameDataWithReferencers& RenameData : AssetsToRename)
 	{
@@ -1440,7 +1467,7 @@ void FAssetRenameManager::SetupPublicAssets(TArray<FAssetRenameDataWithReference
 				else
 				{
 					RenameData.FailureReason = FText::Format(
-						LOCTEXT("AssetCannotBePublic", "Cannot move asset to {0} because it cannot be made public in order to be referenced from a different plugin by {1}"), 
+						LOCTEXT("AssetCannotBePublic", "Cannot move asset to {0} because it cannot be made public in order to be referenced from a different plugin by {1}"),
 						FText::FromStringView(NewPackageMountPoint),
 						FText::FromString(ReferencingAssetInDifferentMountPoint));
 				}
@@ -1495,7 +1522,7 @@ void FAssetRenameManager::SetupPublicAssets(TArray<FAssetRenameDataWithReference
 	}
 }
 
-void FAssetRenameManager::RenameReferencingSoftObjectPaths(const TArray<UPackage *> PackagesToCheck, const TMap<FSoftObjectPath, FSoftObjectPath>& AssetRedirectorMap) const
+void FAssetRenameManager::RenameReferencingSoftObjectPaths(const TArray<UPackage*> PackagesToCheck, const TMap<FSoftObjectPath, FSoftObjectPath>& AssetRedirectorMap) const
 {
 	// Add redirects as needed
 	for (const TPair<FSoftObjectPath, FSoftObjectPath>& Pair : AssetRedirectorMap)
@@ -1646,32 +1673,30 @@ bool FAssetRenameManager::CheckPackageForSoftObjectReferences(UPackage* Package,
 	return bFoundReference;
 }
 
-void FAssetRenameManager::PerformAssetRename(TArray<FAssetRenameDataWithReferencers>& AssetsToRename) const
+void FAssetRenameManager::PerformAssetRename(TArray<FAssetRenameDataWithReferencers>& AssetsToRename, FScopedSlowTask& RenamingSlowTask) const
 {
-	PerformAssetRename(AssetsToRename, TArray<UPackage*>());
+	PerformAssetRename(AssetsToRename, TArray<UPackage*>(), RenamingSlowTask);
 }
 
-void FAssetRenameManager::PerformAssetRename(TArray<FAssetRenameDataWithReferencers>& AssetsToRename, const TArray<UPackage*>& ReferencingPackagesToSave) const
+void FAssetRenameManager::PerformAssetRename(TArray<FAssetRenameDataWithReferencers>& AssetsToRename, const TArray<UPackage*>& ReferencingPackagesToSave, FScopedSlowTask& RenamingSlowTask) const
 {
-	const FText AssetRenameSlowTask = LOCTEXT("AssetRenameSlowTask", "Renaming Assets");
-	GWarn->BeginSlowTask(AssetRenameSlowTask, true);
-
 	/**
 	 * We need to collect and check those cause dependency graph is only
 	 * representing on-disk state and we want to support rename for in-memory
 	 * objects. It is only needed for string references as in memory references
 	 * for other objects are pointers, so renames doesn't apply to those.
 	 */
-	TArray<UPackage *> DirtyPackagesToCheckForSoftReferences;
+	TArray<UPackage*> DirtyPackagesToCheckForSoftReferences;
 
 	FEditorFileUtils::GetDirtyWorldPackages(DirtyPackagesToCheckForSoftReferences);
 	FEditorFileUtils::GetDirtyContentPackages(DirtyPackagesToCheckForSoftReferences);
 
 	TArray<UPackage*> PackagesToSave = ReferencingPackagesToSave;
 	TArray<UPackage*> PotentialPackagesToDelete;
+	float ProgressStep = 2.0f / AssetsToRename.Num();
 	for (int32 AssetIdx = 0; AssetIdx < AssetsToRename.Num(); ++AssetIdx)
 	{
-		GWarn->StatusUpdate(AssetIdx, AssetsToRename.Num(), AssetRenameSlowTask);
+		RenamingSlowTask.EnterProgressFrame(ProgressStep);
 
 		FAssetRenameDataWithReferencers& RenameData = AssetsToRename[AssetIdx];
 
@@ -1682,7 +1707,7 @@ void FAssetRenameManager::PerformAssetRename(TArray<FAssetRenameDataWithReferenc
 		}
 
 		UObject* Asset = RenameData.Asset.Get();
-		TArray<UPackage *> PackagesToCheckForSoftReferences;
+		TArray<UPackage*> PackagesToCheckForSoftReferences;
 
 		bool bIsCaseChangeOnly = false;
 		if (!RenameData.bOnlyFixSoftReferences)
@@ -1713,7 +1738,7 @@ void FAssetRenameManager::PerformAssetRename(TArray<FAssetRenameDataWithReferenc
 				FAssetToolsModule& AssetToolsModule = FModuleManager::LoadModuleChecked<FAssetToolsModule>("AssetTools");
 				FString PackageName;
 				FString AssetName;
-				FString BasePath = RenameData.NewPackagePath + TEXT("/RenameTmp") ;
+				FString BasePath = RenameData.NewPackagePath + TEXT("/RenameTmp");
 				AssetToolsModule.Get().CreateUniqueAssetName(BasePath, TEXT(""), PackageName, AssetName);
 
 				ObjectTools::FPackageGroupName TempPGN;
@@ -1827,7 +1852,7 @@ void FAssetRenameManager::PerformAssetRename(TArray<FAssetRenameDataWithReferenc
 		}
 	}
 
-	GWarn->EndSlowTask();
+	RenamingSlowTask.EnterProgressFrame(0.5f, LOCTEXT("SavingRenamedAssetSlowTask", "Saving renamed assets..."));
 
 	// Save all renamed assets and any redirectors that were left behind
 	if (PackagesToSave.Num() > 0)
@@ -1844,6 +1869,9 @@ void FAssetRenameManager::PerformAssetRename(TArray<FAssetRenameDataWithReferenc
 		ISourceControlModule::Get().QueueStatusUpdate(Filenames);
 	}
 
+	const FText& UpdatingRevisionControlSlowTask = LOCTEXT("UpdatingRevisionControlSlowTask", "Updating Revision Control...");
+	RenamingSlowTask.EnterProgressFrame(0.5f, UpdatingRevisionControlSlowTask);
+
 	// Bulk update SCC status for old packages since it is faster than doing it one by one below
 	if (ISourceControlModule::Get().IsEnabled())
 	{
@@ -1871,9 +1899,15 @@ void FAssetRenameManager::PerformAssetRename(TArray<FAssetRenameDataWithReferenc
 		}
 	}
 
+	RenamingSlowTask.EnterProgressFrame(5.0f, UpdatingRevisionControlSlowTask);
+	FScopedSlowTask BuildingRelationshipSlowTask(5.0f, LOCTEXT("BranchingRenamedAssetsSlowTask", "Building a relationship between the renamed files in Revision Control..."));
+
 	// Now branch the files in source control if possible
+	ProgressStep = 2.9f / AssetsToRename.Num();
 	for (const FAssetRenameDataWithReferencers& RenameData : AssetsToRename)
 	{
+		BuildingRelationshipSlowTask.EnterProgressFrame(ProgressStep);
+
 		UPackage* OldPackage = FindPackage(nullptr, *RenameData.OldObjectPath.GetLongPackageName());
 		UPackage* NewPackage = FindPackage(nullptr, *RenameData.NewObjectPath.GetLongPackageName());
 
@@ -1925,6 +1959,19 @@ void FAssetRenameManager::SaveReferencingPackages(const TArray<UPackage*>& Refer
 	}
 }
 
+void FAssetRenameManager::RenameInterrupted(const TArray<FAssetRenameData>& AssetsToRename, const FText& InterruptionReason, bool bWithDialog) const
+{
+	TArray<FAssetRenameDataWithReferencers> AssetsToRenameWithReferencers;
+	for (const FAssetRenameData& AssetRenameData : AssetsToRename)
+	{
+		FAssetRenameDataWithReferencers AssetToRenameWithReferencers(AssetRenameData);
+		AssetToRenameWithReferencers.bRenameFailed = true;
+		AssetToRenameWithReferencers.FailureReason = InterruptionReason;
+		AssetsToRenameWithReferencers.Add(AssetToRenameWithReferencers);
+	}
+	ReportFailures(AssetsToRenameWithReferencers, bWithDialog);
+}
+
 int32 FAssetRenameManager::ReportFailures(const TArray<FAssetRenameDataWithReferencers>& AssetsToRename, bool bWithDialog) const
 {
 	TArray<FText> FailedRenames;
@@ -1952,7 +1999,9 @@ int32 FAssetRenameManager::ReportFailures(const TArray<FAssetRenameDataWithRefer
 	{
 		if (bWithDialog)
 		{
-			SRenameFailures::OpenRenameFailuresDialog(FailedRenames);
+			SFileListReportDialog::OpenDialog(LOCTEXT("FailedRenamesDialog", "Failed Renames"),
+				LOCTEXT("RenameFailureTitle", "The following assets could not be renamed."),
+				FailedRenames);
 		}
 		else
 		{
@@ -1966,4 +2015,37 @@ int32 FAssetRenameManager::ReportFailures(const TArray<FAssetRenameDataWithRefer
 	return FailedRenames.Num();
 }
 
+bool FAssetRenameManager::RequiresCheckingForVariants(const TArray<FAssetRenameData>& InAssetsToRename, TArray<FAssetRenameData>& OutAssetsAndVariants, TArray<FAssetRenameData>& OutAssetsToCheckForVariants) const
+{
+	const UAssetDefinitionRegistry* AssetDefinitionRegistry = UAssetDefinitionRegistry::Get();
+
+	const UEditorProjectAssetSettings* Settings = GetDefault<UEditorProjectAssetSettings>();
+	bool bSettingsRequiresCheckingForVariants = !(Settings && !Settings->bRenameLocalizedVariantsAlongsideSourceAsset);
+	if (bSettingsRequiresCheckingForVariants)
+	{
+		OutAssetsAndVariants.Reserve(InAssetsToRename.Num());
+		OutAssetsToCheckForVariants.Reserve(InAssetsToRename.Num());
+		for (const FAssetRenameData& AssetAndName : InAssetsToRename)
+		{
+			UClass* CurrentAssetClass = AssetAndName.Asset != nullptr ? AssetAndName.Asset->GetClass() : nullptr;
+			const UAssetDefinition* CurrentAssetDefinition = AssetDefinitionRegistry->GetAssetDefinitionForClass(CurrentAssetClass);
+			bool bAssetRequiresCheckingForVariant = AssetAndName.bAlsoRenameLocalizedVariants && CurrentAssetDefinition != nullptr && CurrentAssetDefinition->CanLocalize(FAssetData()).IsSupported();
+			if (bAssetRequiresCheckingForVariant)
+			{
+				OutAssetsToCheckForVariants.Add(AssetAndName);
+			}
+			else
+			{
+				OutAssetsAndVariants.Add(AssetAndName);
+			}
+		}
+	}
+	else
+	{
+		OutAssetsAndVariants = InAssetsToRename;
+	}
+
+	return !OutAssetsToCheckForVariants.IsEmpty();
+}
+
 #undef LOCTEXT_NAMESPACE
diff --git a/Engine/Source/Developer/AssetTools/Private/AssetRenameManager.h b/Engine/Source/Developer/AssetTools/Private/AssetRenameManager.h
--- a/Engine/Source/Developer/AssetTools/Private/AssetRenameManager.h
+++ b/Engine/Source/Developer/AssetTools/Private/AssetRenameManager.h
@@ -58,12 +58,18 @@ public:
 	/** Filters packages list depending on if it actually has soft object paths pointing to the specific object being renamed */
 	bool CheckPackageForSoftObjectReferences(UPackage* Package, const TMap<FSoftObjectPath, FSoftObjectPath>& AssetRedirectorMap, TMap<FSoftObjectPath, TArray<UObject*>>& OutReferencingObjects) const;
 
+	/** If this returns false, then ExpandAssetsRenameDataWithTheirVariants would just return a copy of its input */
+	bool CanExpandAssetsRenameDataWithTheirVariants(const TArray<FAssetRenameData>& InAssetsAndNames) const;
+
 private:
-	/** Callback used by DiscoverintAssetsDialog to call FixrefrencesAndRename */
-	void FixReferencesAndRenameCallback(TArray<FAssetRenameData> AssetsAndNames, bool bAutoCheckout, bool bWithDialog) const;
+	/** Callback used by DiscoverintAssetsDialog to call RenameAssetsAndVariants */
+	void RenameAssetsAndVariantsCallback(TArray<FAssetRenameData> InAssetsRenameData, bool bAutoCheckout, bool bWithDialog) const;
+
+	/** Check if renaming these assets should also rename localized variants assets (or source asset) */
+	bool RenameAssetsAndVariants(const TArray<FAssetRenameData>& InAssetsRenameData, bool bAutoCheckout, bool bWithDialog) const;
 
 	/** Attempts to load and fix redirector references for the supplied assets */
-	bool FixReferencesAndRename(const TArray<FAssetRenameData>& AssetsAndNames, bool bAutoCheckout, bool bWithDialog) const;
+	bool FixReferencesAndRename(const TArray<FAssetRenameData>& AssetsAndNames, bool bAutoCheckout, bool bWithDialog, FScopedSlowTask& RenamingSlowTask) const;
 
 	/**
 	 * Get lists of assets with references from CDOs
@@ -115,17 +121,23 @@ private:
 	void SetupPublicAssets(TArray<FAssetRenameDataWithReferencers>& AssetsToRename) const;
 
 	/** Performs the asset rename after the user has selected to proceed */
-	void PerformAssetRename(TArray<FAssetRenameDataWithReferencers>& AssetsToRename) const;
+	void PerformAssetRename(TArray<FAssetRenameDataWithReferencers>& AssetsToRename, FScopedSlowTask& RenamingSlowTask) const;
 
 	/** Performs the asset rename after the user has selected to proceed, also saving the provided referencing packages at the same time */
-	void PerformAssetRename(TArray<FAssetRenameDataWithReferencers>& AssetsToRename, const TArray<UPackage*>& ReferencingPackagesToSave) const;
+	void PerformAssetRename(TArray<FAssetRenameDataWithReferencers>& AssetsToRename, const TArray<UPackage*>& ReferencingPackagesToSave, FScopedSlowTask& RenamingSlowTask) const;
 
 	/** Saves all the referencing packages and updates SCC state */
 	void SaveReferencingPackages(const TArray<UPackage*>& ReferencingPackagesToSave) const;
 
+	/** Prematurely interrupt the rename process for the given reason */
+	void RenameInterrupted(const TArray<FAssetRenameData>& AssetsToRename, const FText& InterruptionReason, bool bWithDialog) const;
+
 	/** Report any failures that may have happened during the rename. Return the number of failures */
 	int32 ReportFailures(const TArray<FAssetRenameDataWithReferencers>& AssetsToRename, bool bWithDialog) const;
 
+	/** Internal check to know if we should look for localized variants */
+	bool RequiresCheckingForVariants(const TArray<FAssetRenameData>& InAssetsToRename, TArray<FAssetRenameData>& OutAssetsAndVariants, TArray<FAssetRenameData>& OutAssetsToCheckForVariants) const;
+
 	/** Called when a package is dirtied, clears the cache */
 	void OnMarkPackageDirty(UPackage* Pkg, bool bWasDirty);
 
diff --git a/Engine/Source/Developer/AssetTools/Private/AssetTools.cpp b/Engine/Source/Developer/AssetTools/Private/AssetTools.cpp
--- a/Engine/Source/Developer/AssetTools/Private/AssetTools.cpp
+++ b/Engine/Source/Developer/AssetTools/Private/AssetTools.cpp
@@ -98,6 +98,7 @@
 #include "PackageMigrationContext.h"
 #include "ComponentReregisterContext.h"
 #include "AssetViewUtils.h"
+#include "LocalizedAssetTools.h"
 
 #if WITH_EDITOR
 #include "Subsystems/AssetEditorSubsystem.h"
@@ -1340,6 +1341,7 @@ UAssetTools::UAssetTools(const class FObjectInitializer& ObjectInitializer)
 UAssetToolsImpl::UAssetToolsImpl(const FObjectInitializer& ObjectInitializer)
 	: Super(ObjectInitializer)
 	, AssetRenameManager(MakeShareable(new FAssetRenameManager))
+	, LocalizedAssetTools(MakeShareable(new FLocalizedAssetTools))
 	, AssetFixUpRedirectors(MakeShareable(new FAssetFixUpRedirectors))
 	, NextUserCategoryBit(EAssetTypeCategories::FirstUser)
 	, AssetClassPermissionList_DEPRECATED(MakeShared<FNamePermissionList>())
@@ -1572,21 +1574,14 @@ TWeakPtr<IAssetTypeActions> UAssetToolsImpl::GetAssetTypeActionsForClass(const U
 	return nullptr;
 }
 
-bool UAssetToolsImpl::CanLocalize(const UClass* Class) const
+TSharedPtr<ILocalizedAssetTools> UAssetToolsImpl::GetLocalizedAssetTools() const
 {
-	if (const UAssetDefinition* AssetDefinition = UAssetDefinitionRegistry::Get()->GetAssetDefinitionForClass(Class))
-	{
-		return AssetDefinition->CanLocalize(FAssetData()).IsSupported();
-	}
-	else
-	{
-		if (TSharedPtr<IAssetTypeActions> AssetActions = GetAssetTypeActionsForClass(Class).Pin())
-		{
-			return AssetActions->CanLocalize();
-		}
-	}
+	return LocalizedAssetTools;
+}
 
-	return false;
+bool UAssetToolsImpl::CanLocalize(const UClass* Class) const
+{
+	return LocalizedAssetTools->CanLocalize(Class);
 }
 
 TOptional<FLinearColor> UAssetToolsImpl::GetTypeColor(const UClass* Class) const
diff --git a/Engine/Source/Developer/AssetTools/Private/AssetTools.h b/Engine/Source/Developer/AssetTools/Private/AssetTools.h
--- a/Engine/Source/Developer/AssetTools/Private/AssetTools.h
+++ b/Engine/Source/Developer/AssetTools/Private/AssetTools.h
@@ -69,7 +69,10 @@ public:
 	virtual void UnregisterAssetTypeActions(const TSharedRef<IAssetTypeActions>& ActionsToRemove) override;
 	virtual void GetAssetTypeActionsList( TArray<TWeakPtr<IAssetTypeActions>>& OutAssetTypeActionsList ) const override;
 	virtual TWeakPtr<IAssetTypeActions> GetAssetTypeActionsForClass(const UClass* Class) const override;
-	virtual bool CanLocalize(const UClass* Class) const;
+
+	virtual TSharedPtr<ILocalizedAssetTools> GetLocalizedAssetTools() const override;
+	virtual bool CanLocalize(const UClass* Class) const override;
+
 	virtual TOptional<FLinearColor> GetTypeColor(const UClass* Class) const override;
 	
 	virtual TArray<TWeakPtr<IAssetTypeActions>> GetAssetTypeActionsListForClass(const UClass* Class) const override;
@@ -166,6 +169,9 @@ public:
 	/** The manager to handle renaming assets */
 	TSharedPtr<FAssetRenameManager> AssetRenameManager;
 
+	/** The tools to manage localized variants */
+	TSharedPtr<ILocalizedAssetTools> LocalizedAssetTools;
+
 	/** The manager to handle fixing up redirectors */
 	TSharedPtr<FAssetFixUpRedirectors> AssetFixUpRedirectors;
 private:
diff --git a/Engine/Source/Developer/AssetTools/Private/AssetViewUtils.cpp b/Engine/Source/Developer/AssetTools/Private/AssetViewUtils.cpp
--- a/Engine/Source/Developer/AssetTools/Private/AssetViewUtils.cpp
+++ b/Engine/Source/Developer/AssetTools/Private/AssetViewUtils.cpp
@@ -411,7 +411,9 @@ void AssetViewUtils::MoveAssets(const TArray<UObject*>& Assets, const FString& D
 			PackagePath = DestPath;
 		}
 
-		new(AssetsAndNames) FAssetRenameData(Asset, PackagePath, ObjectName);
+		const bool bSoftReferenceOnly = false;
+		const bool bAlsoRenameLocalizedVariants = true;
+		new(AssetsAndNames) FAssetRenameData(Asset, PackagePath, ObjectName, bSoftReferenceOnly, bAlsoRenameLocalizedVariants);
 	}
 
 	if ( AssetsAndNames.Num() > 0 )
diff --git a/Engine/Source/Developer/AssetTools/Private/LocalizedAssetTools.cpp b/Engine/Source/Developer/AssetTools/Private/LocalizedAssetTools.cpp
new file mode 100644
--- /dev/null
+++ b/Engine/Source/Developer/AssetTools/Private/LocalizedAssetTools.cpp
@@ -0,0 +1,272 @@
+// Copyright Epic Games, Inc. All Rights Reserved.
+
+#include "LocalizedAssetTools.h"
+
+#include "AssetDefinitionRegistry.h"
+#include "AssetRegistry/IAssetRegistry.h"
+#include "AssetTools.h"
+#include "AssetToolsLog.h"
+#include "AssetToolsModule.h"
+#include "Internationalization/PackageLocalizationUtil.h"
+#include "ISourceControlModule.h"
+#include "ISourceControlProvider.h"
+#include "Logging/LogMacros.h"
+#include "Misc/MessageDialog.h"
+#include "SFileListReportDialog.h"
+#include "SourceControlHelpers.h"
+#include "SourceControlPreferences.h"
+#include "UObject/UObjectGlobals.h"
+
+#define LOCTEXT_NAMESPACE "LocalizedAssetTools"
+
+FLocalizedAssetTools::FLocalizedAssetTools() : 
+	RevisionControlIsNotAvailableWarningText(LOCTEXT("RevisionControlIsRequiredToChangeLocalizableAssets", "Revision Control is required to move/rename/delete localizable assets for this project and it is currently not accessible."))
+	, FilesNeedToBeOnDiskWarningText(LOCTEXT("FilesToSyncDialogTitle", "Files in Revision Control need to be on disk"))
+	
+{
+
+}
+
+bool FLocalizedAssetTools::CanLocalize(const UClass* Class) const
+{
+	if (const UAssetDefinition* AssetDefinition = UAssetDefinitionRegistry::Get()->GetAssetDefinitionForClass(Class))
+	{
+		return AssetDefinition->CanLocalize(FAssetData()).IsSupported();
+	}
+	else
+	{
+		FAssetToolsModule& Module = FModuleManager::GetModuleChecked<FAssetToolsModule>("AssetTools");
+		if (TSharedPtr<IAssetTypeActions> AssetActions = Module.Get().GetAssetTypeActionsForClass(Class).Pin())
+		{
+			return AssetActions->CanLocalize();
+		}
+	}
+
+	return false;
+}
+
+ELocalizedAssetsOnDiskResult FLocalizedAssetTools::GetLocalizedVariantsOnDisk(const TArray<FName>& InPackages, TMap<FName, TArray<FName>>& OutLocalizedVariantsBySource, TArray<FName>* OutPackagesNotFound /*= nullptr*/) const
+{
+	FScopedSlowTask GettingLocalizedVariantsOnDiskSlowTask(1.0f, LOCTEXT("GettingLocalizedVariantsOnDiskSlowTask", "Getting localized variants on disk..."));
+
+	OutLocalizedVariantsBySource.Reserve(InPackages.Num());
+	if (OutPackagesNotFound != nullptr)
+	{
+		OutPackagesNotFound->Reserve(InPackages.Num());
+	}
+
+	TMap<FName, FAssetData> PackagesToAssetDataMap;
+	UE::AssetRegistry::GetAssetForPackages(InPackages, PackagesToAssetDataMap);
+	if (!ensureAlwaysMsgf(PackagesToAssetDataMap.Num() == InPackages.Num(), TEXT("PackageNames were not properly converted to AssetData.")))
+	{
+		for (const FName& OriginalAssetName : InPackages)
+		{
+			OutLocalizedVariantsBySource.Add({ OriginalAssetName, TArray<FName>() });
+		}
+		return ELocalizedAssetsOnDiskResult::PackageNamesError;
+	}
+
+	UAssetDefinitionRegistry* AssetDefinitionRegistry = UAssetDefinitionRegistry::Get();
+	float ProgressStep = 1.0f / InPackages.Num();
+	for (const FName& OriginalAssetName : InPackages)
+	{
+		GettingLocalizedVariantsOnDiskSlowTask.EnterProgressFrame(ProgressStep);
+		FString SourceAssetPathStr = OriginalAssetName.ToString();
+		FPackageLocalizationUtil::ConvertLocalizedToSource(SourceAssetPathStr, SourceAssetPathStr);
+		const FName SourceAssetName(*SourceAssetPathStr);
+
+		if (OutLocalizedVariantsBySource.Contains(SourceAssetName))
+		{
+			continue; // We want to avoid doing any unnecessary work if it was already processed
+		}
+
+		// We want to avoid doing any unnecessary work on assets that does not require checking for variants
+		FAssetData SourceAssetData = PackagesToAssetDataMap[OriginalAssetName];
+		UClass* SourceAssetClass = SourceAssetData.GetClass();
+		const UAssetDefinition* SourceAssetDefinition = AssetDefinitionRegistry->GetAssetDefinitionForClass(SourceAssetClass);
+		bool bShouldCheckForVariant = SourceAssetDefinition != nullptr && SourceAssetDefinition->CanLocalize(SourceAssetData).IsSupported();
+		if (!bShouldCheckForVariant)
+		{
+			OutLocalizedVariantsBySource.Add({ SourceAssetName, TArray<FName>() });
+			continue;
+		}
+
+		// Check on disk for localized variants first. Remember the assets that had no variants on disk
+		// because we will then check in Revision Control if applicable
+		TArray<FString> LocalizedVariantsPaths;
+		FPackageLocalizationUtil::GetLocalizedVariantsAbsolutePaths(SourceAssetPathStr, LocalizedVariantsPaths);
+		if (LocalizedVariantsPaths.IsEmpty())
+		{
+			if (OutPackagesNotFound != nullptr)
+			{
+				OutPackagesNotFound->Add(OriginalAssetName);
+			}
+			continue;
+		}
+
+		// If localized variants were found on disk, let's build renaming data for them too
+		TArray<FName> LocalizedAssets;
+		LocalizedAssets.Reserve(LocalizedVariantsPaths.Num());
+		for (const FString& LocalizedVariantPath : LocalizedVariantsPaths)
+		{
+			FString Culture;
+			FPackageLocalizationUtil::ExtractCultureFromLocalized(LocalizedVariantPath, Culture);
+
+			FString LocalizedAsset;
+			FPackageLocalizationUtil::ConvertSourceToLocalized(SourceAssetPathStr, Culture, LocalizedAsset);
+
+			LocalizedAssets.Add(FName(LocalizedAsset));
+		}
+
+		OutLocalizedVariantsBySource.Add({ SourceAssetName, LocalizedAssets });
+	}
+
+	return ELocalizedAssetsOnDiskResult::Success;
+}
+
+ELocalizedAssetsInSCCResult FLocalizedAssetTools::GetLocalizedVariantsInRevisionControl(const TArray<FName>& InPackages, TMap<FName, TArray<FName>>& OutLocalizedVariantsBySource, TArray<FName>* OutPackagesNotFound /*= nullptr*/) const
+{
+	FScopedSlowTask GetLocalizedVariantsInRevisionControlSlowTask(1.0f, LOCTEXT("GetLocalizedVariantsInRevisionControlSlowTask", "Querying Revision Control for localized variants... This could take a long time."));
+	GetLocalizedVariantsInRevisionControlSlowTask.EnterProgressFrame(0.05f);
+
+	OutLocalizedVariantsBySource.Reserve(InPackages.Num());
+
+	// Modify data to be used by USourceControlHelpers
+	TArray<FString> PackagesAsString;
+	PackagesAsString.Reserve(InPackages.Num());
+	for (const FName& InPackageName : InPackages)
+	{
+		PackagesAsString.Add(InPackageName.ToString());
+	}
+
+	// Let's check the packages presence in Revision Control in a single query
+	TArray<FString> LocalizedVariantsInRevisionControl;
+	GetLocalizedVariantsInRevisionControlSlowTask.EnterProgressFrame(0.9f);
+	bool bOutRevisionControlWasNeeded = !GetLocalizedVariantsDepotPaths(PackagesAsString, LocalizedVariantsInRevisionControl);
+
+	// Fill a proper structure with the results
+	float ProgressStep = 0.03f / LocalizedVariantsInRevisionControl.Num();
+	for (const FString& LocalizedVariantInRevisionControl : LocalizedVariantsInRevisionControl)
+	{
+		GetLocalizedVariantsInRevisionControlSlowTask.EnterProgressFrame(ProgressStep);
+		FString SourceAsset;
+		FPackageLocalizationUtil::ConvertToSource(LocalizedVariantInRevisionControl, SourceAsset);
+		FName SourceAssetName(SourceAsset);
+		TArray<FName>* VariantsBySourceFound = OutLocalizedVariantsBySource.Find(SourceAssetName);
+		TArray<FName>& VariantsBySource = (VariantsBySourceFound != nullptr ? *VariantsBySourceFound : OutLocalizedVariantsBySource.Add({ SourceAssetName, TArray<FName>() }));
+		VariantsBySource.Add(FName(LocalizedVariantInRevisionControl));
+	}
+
+	// Don't forget to return the information on the packages that found nothing in Revision Control
+	if (OutPackagesNotFound != nullptr)
+	{
+		ProgressStep = 0.02f / InPackages.Num();
+		for (const FName& PackageName : InPackages)
+		{
+			GetLocalizedVariantsInRevisionControlSlowTask.EnterProgressFrame(ProgressStep);
+			FString SourcePackage;
+			FPackageLocalizationUtil::ConvertToSource(PackageName.ToString(), SourcePackage);
+			if (OutLocalizedVariantsBySource.Find(FName(SourcePackage)) == nullptr)
+			{
+				// Package not found
+				OutPackagesNotFound->Add(PackageName);
+			}
+		}
+	}
+
+	return bOutRevisionControlWasNeeded ? ELocalizedAssetsInSCCResult::RevisionControlNotAvailable : ELocalizedAssetsInSCCResult::Success;
+}
+
+ELocalizedAssetsResult FLocalizedAssetTools::GetLocalizedVariants(const TArray<FName>& InPackages, TMap<FName, TArray<FName>>& OutLocalizedVariantsBySourceOnDisk, bool bAlsoCheckInRevisionControl, TMap<FName, TArray<FName>>& OutLocalizedVariantsBySourceInRevisionControl, TArray<FName>* OutPackagesNotFound /*= nullptr*/) const
+{
+	ELocalizedAssetsResult Result = ELocalizedAssetsResult::Success;
+
+	// Check on disk first
+	// Call the LocalizedAssetTools interface to GetLocalizedVariantsOnDisk of a list of packages
+	TMap<FName, TArray<FName>> VariantsBySources;
+	TArray<FName> VariantsMaybeInRevisionControl;
+	ELocalizedAssetsOnDiskResult DiskResult = GetLocalizedVariantsOnDisk(InPackages, OutLocalizedVariantsBySourceOnDisk, bAlsoCheckInRevisionControl ? &VariantsMaybeInRevisionControl : OutPackagesNotFound);
+	Result = (DiskResult == ELocalizedAssetsOnDiskResult::PackageNamesError ? ELocalizedAssetsResult::PackageNamesError : Result);
+
+	// Check in Revision Control if applicable
+	// Call the LocalizedAssetTools interface to GetLocalizedVariantsInRevisionControl of a list of packages
+	bool bRevisionControlWasNeeded = false;
+	if (!VariantsMaybeInRevisionControl.IsEmpty())
+	{
+		if (Result == ELocalizedAssetsResult::Success)
+		{
+			ELocalizedAssetsInSCCResult SCCResult = GetLocalizedVariantsInRevisionControl(VariantsMaybeInRevisionControl, OutLocalizedVariantsBySourceInRevisionControl, OutPackagesNotFound);
+			Result = (SCCResult == ELocalizedAssetsInSCCResult::RevisionControlNotAvailable ? ELocalizedAssetsResult::RevisionControlNotAvailable : Result);
+		}
+		else
+		{
+			OutPackagesNotFound->Append(VariantsMaybeInRevisionControl);
+		}
+	}
+
+	return Result;
+}
+
+void FLocalizedAssetTools::OpenRevisionControlRequiredDialog() const
+{
+	FText WarningText = RevisionControlIsNotAvailableWarningText;
+	const FText AvoidWarningText = LOCTEXT("HowToFixRevisionControlIsRequiredToChangeLocalizableAssets", "If you want to disable this project option, it is located under:\n\tProject Settings/\n\tEditor/\n\tRevision Control/\n\tRequires Revision Control To Rename Localizable Assets\n\nThis option is there to prevent breaking paths between a source asset and its localized variants if they are not on disk.");
+	FMessageDialog::Open(EAppMsgType::Ok, WarningText.Format(LOCTEXT("RevisionControlIsRequiredToChangeLocalizableAssetsDialog", "{0}\n\n{1}"), WarningText, AvoidWarningText));
+}
+
+void FLocalizedAssetTools::OpenFilesInRevisionControlRequiredDialog(const TArray<FText>& FileList) const
+{
+	SFileListReportDialog::OpenDialog(FilesNeedToBeOnDiskWarningText,
+		LOCTEXT("FilesToSyncDialogHeader", "The following assets were found only in Revision Control. They need to be on your disk to be renamed."),
+		FileList);
+}
+
+const FText& FLocalizedAssetTools::GetRevisionControlIsNotAvailableWarningText() const
+{
+	return RevisionControlIsNotAvailableWarningText;
+}
+
+const FText& FLocalizedAssetTools::GetFilesNeedToBeOnDiskWarningText() const
+{
+	return FilesNeedToBeOnDiskWarningText;
+}
+
+bool FLocalizedAssetTools::GetLocalizedVariantsDepotPaths(const TArray<FString>& InPackagesNames, TArray<FString>& OutLocalizedVariantsPaths) const
+{
+	// Ensure source control system is up and running with a configured provider
+	ISourceControlModule& SCModule = ISourceControlModule::Get();
+	if (!SCModule.IsEnabled())
+	{
+		return false;
+	}
+	ISourceControlProvider& Provider = SCModule.GetProvider();
+	if (!Provider.IsAvailable())
+	{
+		return false;
+	}
+
+	// Other providers don't work for now
+	if (Provider.GetName() == "Perforce")
+	{
+		TArray<FString> LocalizedVariantsRegexPaths;
+		LocalizedVariantsRegexPaths.Reserve(InPackagesNames.Num());
+		for (const FString& InPackageName : InPackagesNames)
+		{
+			FString SourcePackageName;
+			FPackageLocalizationUtil::ConvertToSource(InPackageName, SourcePackageName);
+			FString LocalizedVariantsRegexPath;
+			FPackageLocalizationUtil::ConvertSourceToRegexLocalized(SourcePackageName, LocalizedVariantsRegexPath);
+			LocalizedVariantsRegexPath += FPackageName::GetAssetPackageExtension();
+			LocalizedVariantsRegexPaths.Add(LocalizedVariantsRegexPath);
+		}
+
+		bool bSilent = true;
+		bool bIncludeDeleted = true;
+		USourceControlHelpers::GetFilesInDepotAtPaths(LocalizedVariantsRegexPaths, OutLocalizedVariantsPaths, bIncludeDeleted, bSilent, true);
+	}
+
+	return true;
+}
+
+#undef LOCTEXT_NAMESPACE
+
diff --git a/Engine/Source/Developer/AssetTools/Private/LocalizedAssetTools.h b/Engine/Source/Developer/AssetTools/Private/LocalizedAssetTools.h
new file mode 100644
--- /dev/null
+++ b/Engine/Source/Developer/AssetTools/Private/LocalizedAssetTools.h
@@ -0,0 +1,40 @@
+// Copyright Epic Games, Inc. All Rights Reserved.
+
+#pragma once
+
+#include "Containers/Array.h"
+#include "ILocalizedAssetTools.h"
+#include "Internationalization/Text.h"
+#include "UObject/NameTypes.h"
+
+class FLocalizedAssetTools : public ILocalizedAssetTools
+{
+public:
+	FLocalizedAssetTools();
+	virtual ~FLocalizedAssetTools() {}
+
+	// ILocalizedAssetTools implementations
+
+	virtual bool CanLocalize(const UClass* Class) const override;
+
+	virtual ELocalizedAssetsOnDiskResult GetLocalizedVariantsOnDisk(const TArray<FName>& InPackages, TMap<FName, TArray<FName>>& OutLocalizedVariantsBySource, TArray<FName>* OutPackagesNotFound = nullptr) const override;
+
+	virtual ELocalizedAssetsInSCCResult GetLocalizedVariantsInRevisionControl(const TArray<FName>& InPackages, TMap<FName, TArray<FName>>& OutLocalizedVariantsBySource, TArray<FName>* OutPackagesNotFound = nullptr) const override;
+
+	virtual ELocalizedAssetsResult GetLocalizedVariants(const TArray<FName>& InPackages, TMap<FName, TArray<FName>>& OutLocalizedVariantsBySourceOnDisk, bool bAlsoCheckInRevisionControl, TMap<FName, TArray<FName>>& OutLocalizedVariantsBySourceInRevisionControl, TArray<FName>* OutPackagesNotFound = nullptr) const override;
+
+	virtual void OpenRevisionControlRequiredDialog() const override;
+
+	virtual void OpenFilesInRevisionControlRequiredDialog(const TArray<FText>& FileList) const override;
+
+	virtual const FText& GetRevisionControlIsNotAvailableWarningText() const override;
+
+	virtual const FText& GetFilesNeedToBeOnDiskWarningText() const override;
+
+private:
+	bool GetLocalizedVariantsDepotPaths(const TArray<FString>& InPackagesNames, TArray<FString>& OutLocalizedVariantsPaths) const;
+
+private:
+	FText RevisionControlIsNotAvailableWarningText;
+	FText FilesNeedToBeOnDiskWarningText;
+};
diff --git a/Engine/Source/Developer/AssetTools/Private/SFileListReportDialog.cpp b/Engine/Source/Developer/AssetTools/Private/SFileListReportDialog.cpp
new file mode 100644
--- /dev/null
+++ b/Engine/Source/Developer/AssetTools/Private/SFileListReportDialog.cpp
@@ -0,0 +1,112 @@
+// Copyright Epic Games, Inc. All Rights Reserved.
+
+#include "SFileListReportDialog.h"
+
+#include "Framework/Application/SlateApplication.h"
+#include "Interfaces/IMainFrameModule.h"
+#include "Widgets/Input/SButton.h"
+#include "Widgets/Layout/SBorder.h"
+#include "Widgets/SWindow.h"
+#include "Widgets/Text/STextBlock.h"
+#include "Widgets/Views/ITableRow.h"
+#include "Widgets/Views/SListView.h"
+#include "Widgets/Views/STableViewBase.h"
+
+#define LOCTEXT_NAMESPACE "FileListReportDialog"
+
+void SFileListReportDialog::Construct(const FArguments& InArgs)
+{
+	for (const FText& File : InArgs._Files)
+	{
+		Files.Add(MakeShareable(new FText(File)));
+	}
+
+	ChildSlot
+		[
+			SNew(SBorder)
+				.BorderImage(FAppStyle::GetBrush("Docking.Tab.ContentAreaBrush"))
+				.Padding(FMargin(4, 8, 4, 4))
+				[
+					SNew(SVerticalBox)
+
+						// Title text
+						+ SVerticalBox::Slot()
+						.AutoHeight()
+						[
+							SNew(STextBlock).Text(InArgs._Header)
+						]
+
+						// Files To Sync list
+						+ SVerticalBox::Slot()
+						.Padding(0, 8)
+						.FillHeight(1.f)
+						[
+							SNew(SBorder)
+								.BorderImage(FAppStyle::GetBrush("ToolPanel.GroupBorder"))
+								[
+									SNew(SListView<TSharedRef<FText>>)
+										.ListItemsSource(&Files)
+										.SelectionMode(ESelectionMode::None)
+										.OnGenerateRow(this, &SFileListReportDialog::MakeListViewWidget)
+								]
+						]
+
+						// Close button
+						+ SVerticalBox::Slot()
+						.AutoHeight()
+						.Padding(0, 4)
+						.HAlign(HAlign_Right)
+						[
+							SNew(SButton)
+								.OnClicked(this, &SFileListReportDialog::CloseClicked)
+								.Text(LOCTEXT("WindowCloseButton", "Close"))
+						]
+				]
+		];
+}
+
+void SFileListReportDialog::OpenDialog(const FText& InTitle, const FText& InHeader, const TArray<FText>& InFiles)
+{
+	TSharedRef<SWindow> RenameWindow = SNew(SWindow)
+		.Title(InTitle)
+		.ClientSize(FVector2D(800, 400))
+		.SupportsMaximize(false)
+		.SupportsMinimize(false)
+		[
+			SNew(SFileListReportDialog).Header(InHeader).Files(InFiles)
+		];
+
+	IMainFrameModule& MainFrameModule = FModuleManager::LoadModuleChecked<IMainFrameModule>(TEXT("MainFrame"));
+
+	if (MainFrameModule.GetParentWindow().IsValid())
+	{
+		FSlateApplication::Get().AddWindowAsNativeChild(RenameWindow, MainFrameModule.GetParentWindow().ToSharedRef());
+	}
+	else
+	{
+		FSlateApplication::Get().AddWindow(RenameWindow);
+	}
+}
+
+TSharedRef<ITableRow> SFileListReportDialog::MakeListViewWidget(TSharedRef<FText> Item, const TSharedRef<STableViewBase>& OwnerTable)
+{
+	return
+		SNew(STableRow< TSharedRef<FText> >, OwnerTable)
+		[
+			SNew(STextBlock).Text(Item.Get())
+		];
+}
+
+FReply SFileListReportDialog::CloseClicked()
+{
+	TSharedPtr<SWindow> Window = FSlateApplication::Get().FindWidgetWindow(AsShared());
+
+	if (Window.IsValid())
+	{
+		Window->RequestDestroyWindow();
+	}
+
+	return FReply::Handled();
+}
+
+#undef LOCTEXT_NAMESPACE
diff --git a/Engine/Source/Developer/AssetTools/Private/SFileListReportDialog.h b/Engine/Source/Developer/AssetTools/Private/SFileListReportDialog.h
new file mode 100644
--- /dev/null
+++ b/Engine/Source/Developer/AssetTools/Private/SFileListReportDialog.h
@@ -0,0 +1,33 @@
+// Copyright Epic Games, Inc. All Rights Reserved.
+
+#pragma once
+
+#include "Input/Reply.h"
+#include "Widgets/SCompoundWidget.h"
+
+class ITableRow;
+class STableViewBase;
+
+class SFileListReportDialog : public SCompoundWidget
+{
+public:
+	SLATE_BEGIN_ARGS(SFileListReportDialog) {}
+
+		SLATE_ARGUMENT(FText, Header)
+		SLATE_ARGUMENT(TArray<FText>, Files)
+
+	SLATE_END_ARGS()
+
+	void Construct(const FArguments& InArgs);
+
+	static void OpenDialog(const FText& InTitle, const FText& InHeader, const TArray<FText>& InFiles);
+
+private:
+	TSharedRef<ITableRow> MakeListViewWidget(TSharedRef<FText> Item, const TSharedRef<STableViewBase>& OwnerTable);
+
+	FReply CloseClicked();
+
+private:
+	FText Header;
+	TArray< TSharedRef<FText> > Files;
+};
diff --git a/Engine/Source/Developer/AssetTools/Public/IAssetTools.h b/Engine/Source/Developer/AssetTools/Public/IAssetTools.h
--- a/Engine/Source/Developer/AssetTools/Public/IAssetTools.h
+++ b/Engine/Source/Developer/AssetTools/Public/IAssetTools.h
@@ -21,6 +21,7 @@ struct FAssetData;
 class IAssetTools;
 class IAssetTypeActions;
 class IClassTypeActions;
+class ILocalizedAssetTools;
 class UFactory;
 class UAssetImportTask;
 class UAdvancedCopyCustomization;
@@ -106,9 +107,14 @@ struct FAssetRenameData
 	/** If true, only fix soft references. This will work even if Asset is null because it has already been renamed */
 	UPROPERTY()
 	bool bOnlyFixSoftReferences;
+	
+	/** If true, will also try to rename all localized variants with corresponding paths */
+	UPROPERTY()
+	bool bAlsoRenameLocalizedVariants;
 
 	FAssetRenameData()
 		: bOnlyFixSoftReferences(false)
+		, bAlsoRenameLocalizedVariants(false) // This should probably always be true but it is false for backward compatibility reasons
 	{}
 
 	/** These constructors leave some fields empty, they are fixed up inside AssetRenameManager */
@@ -117,13 +123,24 @@ struct FAssetRenameData
 		, NewPackagePath(InNewPackagePath)
 		, NewName(InNewName)
 		, bOnlyFixSoftReferences(false)
+		, bAlsoRenameLocalizedVariants(false) // This should probably always be true but it is false for backward compatibility reasons
+	{
+	}
+	/** These constructors leave some fields empty, they are fixed up inside AssetRenameManager */
+	FAssetRenameData(const TWeakObjectPtr<UObject>& InAsset, const FString& InNewPackagePath, const FString& InNewName, bool bInOnlyFixSoftReferences, bool bInAlsoRenameLocalizedVariants)
+		: Asset(InAsset)
+		, NewPackagePath(InNewPackagePath)
+		, NewName(InNewName)
+		, bOnlyFixSoftReferences(bInOnlyFixSoftReferences)
+		, bAlsoRenameLocalizedVariants(bInAlsoRenameLocalizedVariants)
 	{
 	}
 	
-	FAssetRenameData(const FSoftObjectPath& InOldObjectPath, const FSoftObjectPath& InNewObjectPath, bool bInOnlyFixSoftReferences = false)
+	FAssetRenameData(const FSoftObjectPath& InOldObjectPath, const FSoftObjectPath& InNewObjectPath, bool bInOnlyFixSoftReferences = false, bool bInAlsoRenameLocalizedVariants = false)
 		: OldObjectPath(InOldObjectPath)
 		, NewObjectPath(InNewObjectPath)
 		, bOnlyFixSoftReferences(bInOnlyFixSoftReferences)
+		, bAlsoRenameLocalizedVariants(bInAlsoRenameLocalizedVariants) // This should probably be true by default but it is false for backward compatibility reasons
 	{
 	}
 };
@@ -272,6 +289,8 @@ public:
 	/** Gets the appropriate AssetTypeActions for the supplied class */
 	virtual TWeakPtr<IAssetTypeActions> GetAssetTypeActionsForClass(const UClass* Class) const = 0;
 
+	virtual TSharedPtr<ILocalizedAssetTools> GetLocalizedAssetTools() const = 0;
+
 	virtual bool CanLocalize(const UClass* Class) const = 0;
 
 	virtual TOptional<FLinearColor> GetTypeColor(const UClass* Class) const = 0;
diff --git a/Engine/Source/Developer/AssetTools/Public/ILocalizedAssetTools.h b/Engine/Source/Developer/AssetTools/Public/ILocalizedAssetTools.h
new file mode 100644
--- /dev/null
+++ b/Engine/Source/Developer/AssetTools/Public/ILocalizedAssetTools.h
@@ -0,0 +1,58 @@
+// Copyright Epic Games, Inc. All Rights Reserved.
+
+#pragma once
+
+#include "UObject/ObjectMacros.h"
+
+enum class ELocalizedAssetsOnDiskResult : uint8
+{
+	/** Getting localized assets on disk succeeded */
+	Success,
+	/** PackageNames could not be converted to FAssetData */
+	PackageNamesError,
+};
+
+enum class ELocalizedAssetsInSCCResult : uint8
+{
+	/** Getting localized assets in Revision Control succeeded */
+	Success,
+	/** Revision Control is required but unavailable */
+	RevisionControlNotAvailable,
+};
+
+enum class ELocalizedAssetsResult : uint8
+{
+	/** Getting localized assets in Revision Control succeeded */
+	Success,
+	/** PackageNames could not be converted to FAssetData */
+	PackageNamesError,
+	/** Revision Control is required but unavailable */
+	RevisionControlNotAvailable,
+};
+
+class ILocalizedAssetTools
+{
+public:
+	virtual bool CanLocalize(const UClass* Class) const = 0;
+
+	/** Get localized variants on disk from an asset name list */
+	virtual ELocalizedAssetsOnDiskResult GetLocalizedVariantsOnDisk(const TArray<FName>& InPackages, TMap<FName, TArray<FName>>& OutLocalizedVariantsBySource, TArray<FName>* OutPackagesNotFound = nullptr) const = 0;
+
+	/** Get localized variants in Revision Control from an asset name list */
+	virtual ELocalizedAssetsInSCCResult GetLocalizedVariantsInRevisionControl(const TArray<FName>& InPackages, TMap<FName, TArray<FName>>& OutLocalizedVariantsBySource, TArray<FName>* OutPackagesNotFound = nullptr) const = 0;
+
+	/** Get localized variants on disk then in Revision Control if Project Settings requires it */
+	virtual ELocalizedAssetsResult GetLocalizedVariants(const TArray<FName>& InPackages, TMap<FName, TArray<FName>>& OutLocalizedVariantsBySourceOnDisk, bool bAlsoCheckInRevisionControl, TMap<FName, TArray<FName>>& OutLocalizedVariantsBySourceInRevisionControl, TArray<FName>* OutPackagesNotFound = nullptr) const = 0;
+
+	/** Open a dialog to required Revision Control to be configured */
+	virtual void OpenRevisionControlRequiredDialog() const = 0;
+
+	/** Open a dialog to required files in Revision Control to be on disk */
+	virtual void OpenFilesInRevisionControlRequiredDialog(const TArray<FText>& FileList) const = 0;
+
+	/** Get shared text to warn that some files need to be on disk (not only in Revision Control) */
+	virtual const FText& GetFilesNeedToBeOnDiskWarningText() const = 0;
+
+	/** Get shared text to warn that Revision Control needs to be available */
+	virtual const FText& GetRevisionControlIsNotAvailableWarningText() const = 0;
+};
diff --git a/Engine/Source/Developer/SourceControl/Private/SourceControlHelpers.cpp b/Engine/Source/Developer/SourceControl/Private/SourceControlHelpers.cpp
--- a/Engine/Source/Developer/SourceControl/Private/SourceControlHelpers.cpp
+++ b/Engine/Source/Developer/SourceControl/Private/SourceControlHelpers.cpp
@@ -18,6 +18,7 @@
 #include "Misc/PackageName.h"
 #include "Logging/MessageLog.h"
 #include "AssetRegistry/AssetRegistryModule.h"
+#include "Internationalization/PackageLocalizationUtil.h"
 
 #include UE_INLINE_GENERATED_CPP_BY_NAME(SourceControlHelpers)
 
@@ -1696,9 +1697,16 @@ void USourceControlHelpers::AsyncQueryFileState(FQueryFileStateDelegate FileStat
 	}
 }
 
-bool USourceControlHelpers::GetFilesInDepotAtPath(const FString& PathToDirectory, TArray<FString>& OutFilesList, bool bIncludeDeleted, bool bSilent)
+bool USourceControlHelpers::GetFilesInDepotAtPath(const FString& Path, TArray<FString>& OutFilesList, bool bIncludeDeleted, bool bSilent, bool bIsFileRegexSearch)
 {
-	TRACE_CPUPROFILER_EVENT_SCOPE(USourceControlHelpers::GetFilesInDepotAtPath);
+	TArray<FString> Paths;
+	Paths.Add(Path);
+	return GetFilesInDepotAtPaths(Paths, OutFilesList, bIncludeDeleted, bSilent, bIsFileRegexSearch);
+}
+
+bool USourceControlHelpers::GetFilesInDepotAtPaths(const TArray<FString>& Paths, TArray<FString>& OutFilesList, bool bIncludeDeleted, bool bSilent, bool bIsFileRegexSearch)
+{
+	TRACE_CPUPROFILER_EVENT_SCOPE(USourceControlHelpers::GetFilesInDepotAtPaths);
 
 	bool bSuccess = false;
 	if (ISourceControlModule::Get().IsEnabled())
@@ -1706,25 +1714,34 @@ bool USourceControlHelpers::GetFilesInDepotAtPath(const FString& PathToDirectory
 		ISourceControlProvider& SourceControlProvider = ISourceControlModule::Get().GetProvider();
 		if (SourceControlProvider.IsAvailable())
 		{
-			FString CorrectedPath = PathToDirectory.Replace(TEXT("/Game"), TEXT(""), ESearchCase::CaseSensitive);
-			FString FullPath = FPaths::ProjectContentDir() / CorrectedPath;
-			FPaths::RemoveDuplicateSlashes(FullPath);
+			TArray<FString> PathsToQuery;
+			for (const FString& Path : Paths)
+			{
+				FString RelativePath;
+				FPackageName::TryConvertLongPackageNameToFilename(Path, RelativePath);
+				FPaths::RemoveDuplicateSlashes(RelativePath);
+				PathsToQuery.Add(RelativePath);
+			}
 
-			TArray<FString> FileArray;
-			FileArray.Add(FullPath);
-			
 			TSharedRef<FGetFileList, ESPMode::ThreadSafe> Operation = ISourceControlOperation::Create<FGetFileList>();
 			Operation->SetIncludeDeleted(bIncludeDeleted);
-			Operation->SetSearchPattern(PathToDirectory);
+			Operation->SetSearchPattern(PathsToQuery);
+			if (bIsFileRegexSearch)
+			{
+				Operation->SetMethodUsed(FGetFileList::EGetFileListMethod::FileRegexSearch);
+			}
 
-			ECommandResult::Type Result = SourceControlProvider.Execute(Operation, FileArray, EConcurrency::Synchronous);
+			ECommandResult::Type Result = SourceControlProvider.Execute(Operation, PathsToQuery, EConcurrency::Synchronous);
 			bSuccess = (Result == ECommandResult::Succeeded);
 
 			if (!bSuccess)
 			{
-				FFormatNamedArguments Arguments;
-				Arguments.Add(TEXT("PathToDirectory"), FText::FromString(PathToDirectory));
-				SourceControlHelpersInternal::LogError(FText::Format(LOCTEXT("CouldNotGetFileList", "Could not get file list under path: {PathToDirectory}."), Arguments), bSilent);
+				for (const FString& Path : Paths)
+				{
+					FFormatNamedArguments Arguments;
+					Arguments.Add(TEXT("Path"), FText::FromString(Path));
+					SourceControlHelpersInternal::LogError(FText::Format(LOCTEXT("CouldNotGetFileList", "Could not get file list at path: {Path}."), Arguments), bSilent);
+				}
 			}
 			else
 			{
diff --git a/Engine/Source/Developer/SourceControl/Private/SourceControlPreferences.cpp b/Engine/Source/Developer/SourceControl/Private/SourceControlPreferences.cpp
--- a/Engine/Source/Developer/SourceControl/Private/SourceControlPreferences.cpp
+++ b/Engine/Source/Developer/SourceControl/Private/SourceControlPreferences.cpp
@@ -2,6 +2,9 @@
 
 #include "SourceControlPreferences.h"
 
+#include "ISourceControlModule.h"
+#include "ISourceControlProvider.h"
+
 #include UE_INLINE_GENERATED_CPP_BY_NAME(SourceControlPreferences)
 
 bool USourceControlPreferences::IsValidationTagEnabled()
@@ -19,3 +22,17 @@ bool USourceControlPreferences::AreUncontrolledChangelistsEnabled()
 	return GetDefault<USourceControlPreferences>()->bEnableUncontrolledChangelists;
 }
 
+bool USourceControlPreferences::RequiresRevisionControlToRenameLocalizableAssets()
+{
+	// This feature only works for Perforce (assume it is Perforce until Revision Control is configured)
+	bool bRevisionControlIsPerforce = true;
+	ISourceControlModule& SCModule = ISourceControlModule::Get();
+	if (SCModule.IsEnabled())
+	{
+		ISourceControlProvider& SCProvider = SCModule.GetProvider();
+		bRevisionControlIsPerforce = SCProvider.GetName() == "Perforce";
+	}
+
+	return bRevisionControlIsPerforce && GetDefault<USourceControlPreferences>()->bRequiresRevisionControlToRenameLocalizableAssets;
+}
+
diff --git a/Engine/Source/Developer/SourceControl/Public/SourceControlHelpers.h b/Engine/Source/Developer/SourceControl/Public/SourceControlHelpers.h
--- a/Engine/Source/Developer/SourceControl/Public/SourceControlHelpers.h
+++ b/Engine/Source/Developer/SourceControl/Public/SourceControlHelpers.h
@@ -494,18 +494,32 @@ public:
 	*/
 	UFUNCTION(BlueprintCallable, Category = "Editor Scripting | Editor Revision Control Helpers", meta = (Keywords = "Source Control"))
 	static void AsyncQueryFileState(FQueryFileStateDelegate FileStateCallback, const FString& InFile, bool bSilent = false);
-
+	
 	/**
 	 * Use currently set source control provider to query the list of files in the depot under a certain path.
 	 * @note	Blocks until action is complete.
 	 *
-	 * @param	PathToDirectory	The path which we want to query the list of files from.
-	 * @param	OutFilesList	An array containing the list of files under the queried path.
-	 * @param	bIncludeDeleted	Include files that have been deleted from the depot.
-	 * @param	bSilent			if false (default) then write out any error info to the Log. Any error text can be retrieved by LastErrorMsg() regardless.
+	 * @param	Path				The path which we want to query the list of files from.
+	 * @param	OutFilesList		An array containing the list of files under the queried path.
+	 * @param	bIncludeDeleted		Include files that have been deleted from the depot.
+	 * @param	bSilent				If false (default) then write out any error info to the Log. Any error text can be retrieved by LastErrorMsg() regardless.
+	 * @param	bIsFileRegexSearch	If false, it is a folder search. If true, it is a file regex search
+	 * @return	Success or failure of the operation
+	 */
+	static bool GetFilesInDepotAtPath(const FString& Path, TArray<FString>& OutFilesList, bool bIncludeDeleted = false, bool bSilent = false, bool bIsFileRegexSearch = false);
+
+	/**
+	 * Use currently set source control provider to query the list of files in the depot under a list of paths.
+	 * @note	Blocks until action is complete.
+	 *
+	 * @param	Paths				The paths which we want to query the list of files from.
+	 * @param	OutFilesList		An array containing the list of files under the queried path.
+	 * @param	bIncludeDeleted		Include files that have been deleted from the depot.
+	 * @param	bSilent				If false (default) then write out any error info to the Log. Any error text can be retrieved by LastErrorMsg() regardless.
+	 * @param	bIsFileRegexSearch	If false, it is a folder search. If true, it is a file regex search
 	 * @return	Success or failure of the operation
 	 */
-	static bool GetFilesInDepotAtPath(const FString& PathToDirectory, TArray<FString>& OutFilesList, bool bIncludeDeleted = false, bool bSilent = false);
+	static bool GetFilesInDepotAtPaths(const TArray<FString>& Paths, TArray<FString>& OutFilesList, bool bIncludeDeleted = false, bool bSilent = false, bool bIsFileRegexSearch = false);
 
 	/**
 	 * Helper function to get a filename for a package name.
diff --git a/Engine/Source/Developer/SourceControl/Public/SourceControlOperations.h b/Engine/Source/Developer/SourceControl/Public/SourceControlOperations.h
--- a/Engine/Source/Developer/SourceControl/Public/SourceControlOperations.h
+++ b/Engine/Source/Developer/SourceControl/Public/SourceControlOperations.h
@@ -144,11 +144,17 @@ public:
 };
 
 /**
- * Operation used to get the file list of a folder out of source control
+ * Operation used to get a file list out of source control
  */
 class FGetFileList : public FSourceControlOperationBase
 {
 public:
+	enum EGetFileListMethod
+	{
+		FolderSearch,
+		FileRegexSearch // Note: Only works for Perforce for now
+	};
+
 	// ISourceControlOperation interface
 	virtual FName GetName() const override
 	{
@@ -170,6 +176,16 @@ public:
 		return bIncludeDeleted;
 	}
 
+	void SetMethodUsed(const EGetFileListMethod& InMethodUsed)
+	{
+		MethodUsed = InMethodUsed;
+	}
+
+	EGetFileListMethod GetMethodUsed()
+	{
+		return MethodUsed;
+	}
+
 	const TArray<FString>& GetFilesList() const
 	{
 		return FilesList;
@@ -180,25 +196,28 @@ public:
 		FilesList = MoveTemp(InFilesList);
 	}
 
-	const FString& GetSearchPattern() const
+	const TArray<FString>& GetSearchPatterns() const
 	{
-		return SearchPattern;
+		return SearchPatterns;
 	}
 
-	void SetSearchPattern(const FString& InSearchPattern)
+	void SetSearchPattern(const TArray<FString>& InSearchPatterns)
 	{
-		SearchPattern = InSearchPattern;
+		SearchPatterns = InSearchPatterns;
 	}
 
 protected:
+	/** Method used to get file list */
+	EGetFileListMethod MethodUsed = EGetFileListMethod::FolderSearch;
+
 	/** Include deleted files in the list. */
 	bool bIncludeDeleted = false;
 
 	/** Stored result of the operation */
 	TArray<FString> FilesList;
 
-	/** The search pattern for the file list */
-	FString SearchPattern;
+	/** The search patterns for the file list */
+	TArray<FString> SearchPatterns;
 };
 
 /**
diff --git a/Engine/Source/Developer/SourceControl/Public/SourceControlPreferences.h b/Engine/Source/Developer/SourceControl/Public/SourceControlPreferences.h
--- a/Engine/Source/Developer/SourceControl/Public/SourceControlPreferences.h
+++ b/Engine/Source/Developer/SourceControl/Public/SourceControlPreferences.h
@@ -32,6 +32,9 @@ public:
 	/** Helper to access easily Enable Uncontrolled Changelists setting */
 	static bool AreUncontrolledChangelistsEnabled();
 
+	/** Helper to access easily Requires Revision Control To Rename Localizable Assets setting */
+	static bool RequiresRevisionControlToRenameLocalizableAssets();
+
 public:
 	/** If enabled, adds a tag in changelist descriptions when they are validated */
 	UPROPERTY(config, EditAnywhere, Category = SourceControl, meta = (ToolTip = "Adds validation tag to changelist description on submit."))
@@ -52,4 +55,8 @@ public:
 	/** Map of collection names and additional text to apply to changelist descriptions when checking them in */
 	UPROPERTY(config, EditAnywhere, Category = SourceControl, meta=(MultiLine=true))
 	TMap<FName, FString> SpecificCollectionChangelistTags;
+
+	/** Revision Control is a requirement to move/rename/delete localizable assets */
+	UPROPERTY(config, EditAnywhere, Category = Internationalization, meta = (ToolTip = "If the project uses Localization and its localized content might not be on disk (but still in your Revision Control), then enabling this feature will require a Revision Control Provider to be configured to move/rename/delete files to be sure no localized variants become orphaned. (Works only with Perforce for now)"))
+	bool bRequiresRevisionControlToRenameLocalizableAssets = false;
 };
diff --git a/Engine/Source/Editor/UnrealEd/Public/Settings/EditorProjectSettings.h b/Engine/Source/Editor/UnrealEd/Public/Settings/EditorProjectSettings.h
--- a/Engine/Source/Editor/UnrealEd/Public/Settings/EditorProjectSettings.h
+++ b/Engine/Source/Editor/UnrealEd/Public/Settings/EditorProjectSettings.h
@@ -285,6 +285,12 @@ public:
 	 */
 	UPROPERTY(EditAnywhere, config, Category = Redirectors)
 	bool bPromptToDeleteUnreferencedRedirectors = true;
+	
+	/**
+	 * Rename localized variants alongside their source asset
+	 */
+	UPROPERTY(EditAnywhere, config, Category = Internationalization)
+	bool bRenameLocalizedVariantsAlongsideSourceAsset = true;
 };
 
 UCLASS(config = Editor, meta = (DisplayName = "Derived Data"), defaultconfig, MinimalAPI)
diff --git a/Engine/Source/Runtime/CoreUObject/Private/Internationalization/PackageLocalizationUtil.cpp b/Engine/Source/Runtime/CoreUObject/Private/Internationalization/PackageLocalizationUtil.cpp
--- a/Engine/Source/Runtime/CoreUObject/Private/Internationalization/PackageLocalizationUtil.cpp
+++ b/Engine/Source/Runtime/CoreUObject/Private/Internationalization/PackageLocalizationUtil.cpp
@@ -1,47 +1,111 @@
 // Copyright Epic Games, Inc. All Rights Reserved.
 
 #include "Internationalization/PackageLocalizationUtil.h"
+#include "AssetRegistry/AssetData.h"
+#include "GenericPlatform/GenericPlatformFile.h"
 #include "Misc/PackageName.h"
+#include "Misc/Paths.h"
+#include "HAL/FileManager.h"
 
-bool FPackageLocalizationUtil::ConvertLocalizedToSource(const FString& InLocalized, FString& OutSource)
+namespace FPackageLocalizationUtilInternal
 {
-	if (FPackageName::IsLocalizedPackage(InLocalized))
+	class L10NCulturesOnDiskVisitor : public IPlatformFile::FDirectoryVisitor
 	{
-		const int32 L10NStart = InLocalized.Find(TEXT("L10N"), ESearchCase::IgnoreCase);
-		if (L10NStart != INDEX_NONE)
+	public:
+		L10NCulturesOnDiskVisitor(const FStringView& InPackageName)
 		{
-			// /Game/L10N/fr/Folder/MyAsset
-			//       ^ We matched here, so we need to walk over the L10N folder, and then walk over the culture code to find the range of characters to remove
-			int32 NumSlashesToFind = 2;
-			int32 L10NEnd = L10NStart;
-			for (; L10NEnd < InLocalized.Len() && NumSlashesToFind > 0; ++L10NEnd)
+			FString PackageName(InPackageName);
+			FString MountPoint = FPackageName::GetPackageMountPoint(PackageName).ToString();
+			FString MountPointRelativePath;
+			FPackageName::TryConvertLongPackageNameToFilename(TEXT("/") + MountPoint, MountPointRelativePath);
+			MountPointAbsolutePath = FPaths::ConvertRelativePathToFull(MountPointRelativePath);
+		}
+
+		const TArray<FString>& GetAllL10NCulturesIDOnDisk()
+		{
+			LazyVisit();
+			return CulturesIDVisited;
+		}
+
+	private:
+		void LazyVisit()
+		{
+			if (!bVisitedAlready)
 			{
-				if (InLocalized[L10NEnd] == TEXT('/'))
-				{
-					--NumSlashesToFind;
-				}
+				bVisitedAlready = true;
+				FString L10NAbsolutePath = MountPointAbsolutePath + TEXT("/L10N/");
+				IFileManager::Get().IterateDirectory(*L10NAbsolutePath, *this);
 			}
+		}
 
-			OutSource = InLocalized;
-			OutSource.RemoveAt(L10NStart, L10NEnd - L10NStart, EAllowShrinking::No);
+		virtual bool Visit(const TCHAR* FilenameOrDirectory, bool bIsDirectory)
+		{
+			FString CultureAbsoluteDirectory(FilenameOrDirectory);
+			int CultureSlashIndex = -1;
+			CultureAbsoluteDirectory.FindLastChar('/', CultureSlashIndex);
+			CulturesIDVisited.Add(CultureAbsoluteDirectory.RightChop(CultureSlashIndex + 1));
 			return true;
 		}
+
+	private:
+		FString MountPointAbsolutePath;
+
+		bool bVisitedAlready = false;
+		TArray<FString> CulturesIDVisited;
+	};
+}
+
+bool FPackageLocalizationUtil::ConvertLocalizedToSource(const FStringView& InLocalized, FString& OutSource)
+{
+	int32 L10NStart = InLocalized.Find(TEXT("/L10N/"), ESearchCase::IgnoreCase);
+	if (L10NStart != INDEX_NONE)
+	{
+		// .../L10N/fr/...
+		//    ^ We matched here, so we need to walk over the L10N folder, and then walk over the culture code to find the range of characters to remove
+		++L10NStart;
+		int32 CultureStart = L10NStart + 5;
+		int32 CultureEnd = CultureStart;
+		for (; CultureEnd < InLocalized.Len(); ++CultureEnd)
+		{
+			if (InLocalized[CultureEnd] == TEXT('/'))
+			{
+				break;
+			}
+		}
+
+		if (CultureEnd == InLocalized.Len())
+		{
+			return false;
+		}
+
+		OutSource = InLocalized;
+		OutSource.RemoveAt(L10NStart, CultureEnd - L10NStart + 1, EAllowShrinking::No);
+		return true;
 	}
 
 	return false;
 }
 
-bool FPackageLocalizationUtil::ConvertSourceToLocalized(const FString& InSource, const FString& InCulture, FString& OutLocalized)
+void FPackageLocalizationUtil::ConvertToSource(const FStringView& InPath, FString& OutSource)
+{
+	if (!ConvertLocalizedToSource(InPath, OutSource))
+	{
+		OutSource = InPath;
+	}
+}
+
+bool FPackageLocalizationUtil::ConvertSourceToLocalized(const FStringView& InSource, const FStringView& InCulture, FString& OutLocalized)
 {
 	if (!FPackageName::IsLocalizedPackage(InSource))
 	{
 		if (InSource.Len() > 0 && InSource[0] == TEXT('/'))
 		{
-			const int32 RootPathEnd = InSource.Find(TEXT("/"), ESearchCase::IgnoreCase, ESearchDir::FromStart, 1);
+			const int32 RootPathEnd = InSource.Find(TEXT("/"), 1, ESearchCase::IgnoreCase);
 			if (RootPathEnd != INDEX_NONE)
 			{
+				FString Culture(InCulture);
 				OutLocalized = InSource;
-				OutLocalized.InsertAt(RootPathEnd, TEXT("/L10N") / InCulture);
+				OutLocalized.InsertAt(RootPathEnd, TEXT("/L10N") / Culture);
 				return true;
 			}
 		}
@@ -50,6 +114,11 @@ bool FPackageLocalizationUtil::ConvertSourceToLocalized(const FString& InSource,
 	return false;
 }
 
+bool FPackageLocalizationUtil::ConvertSourceToRegexLocalized(const FStringView& InSource, FString& OutLocalized)
+{
+	return ConvertSourceToLocalized(InSource, TEXT("*"), OutLocalized);
+}
+
 bool FPackageLocalizationUtil::GetLocalizedRoot(const FString& InPath, const FString& InCulture, FString& OutLocalized)
 {
 	if (InPath.Len() > 0 && InPath[0] == TEXT('/'))
@@ -69,3 +138,49 @@ bool FPackageLocalizationUtil::GetLocalizedRoot(const FString& InPath, const FSt
 
 	return false;
 }
+
+bool FPackageLocalizationUtil::ExtractCultureFromLocalized(const FStringView& InLocalized, FString& OutCulture)
+{
+	const int32 L10NStart = InLocalized.Find(TEXT("/L10N/"), ESearchCase::IgnoreCase);
+	if (L10NStart != INDEX_NONE)
+	{
+		// .../L10N/fr/...
+		//    ^ We matched here, so we need to walk over the L10N folder, and then walk over the culture code
+		const int32 CultureStart = L10NStart + 6;
+		int32 CultureEnd = CultureStart;
+		for (; CultureEnd < InLocalized.Len(); ++CultureEnd)
+		{
+			if (InLocalized[CultureEnd] == TEXT('/'))
+			{
+				break;
+			}
+		}
+
+		if (CultureEnd == InLocalized.Len())
+		{
+			return false;
+		}
+
+		OutCulture = InLocalized.Mid(CultureStart, CultureEnd - CultureStart);
+		return true;
+	}
+
+	return false;
+}
+
+void FPackageLocalizationUtil::GetLocalizedVariantsAbsolutePaths(const FStringView& InSource, TArray<FString>& OutLocalizedAbsolutePaths)
+{
+	FPackageLocalizationUtilInternal::L10NCulturesOnDiskVisitor L10NCulturesOnDisk(InSource);
+	const TArray<FString>& CulturesID = L10NCulturesOnDisk.GetAllL10NCulturesIDOnDisk();
+
+	for (const FString& CultureID : CulturesID)
+	{
+		FString LocalizedVariantPackage;
+		ConvertSourceToLocalized(InSource, CultureID, LocalizedVariantPackage);
+		FString OutFilename;
+		if (FPackageName::DoesPackageExist(LocalizedVariantPackage, &OutFilename))
+		{
+			OutLocalizedAbsolutePaths.Add(FPaths::ConvertRelativePathToFull(OutFilename));
+		}
+	}
+}
diff --git a/Engine/Source/Runtime/CoreUObject/Public/Internationalization/PackageLocalizationUtil.h b/Engine/Source/Runtime/CoreUObject/Public/Internationalization/PackageLocalizationUtil.h
--- a/Engine/Source/Runtime/CoreUObject/Public/Internationalization/PackageLocalizationUtil.h
+++ b/Engine/Source/Runtime/CoreUObject/Public/Internationalization/PackageLocalizationUtil.h
@@ -9,27 +9,47 @@
 struct FPackageLocalizationUtil
 {
 	/**
-	 * Converts a localized version of a package path to the source version (by removing /L10N/<code> from the package path, if present)
+	 * Converts a localized version of a package path to the source version (by removing /L10N/<code> from the package path)
 	 * Note: This does not check whether the source package exists
 	 *
 	 * @param InLocalized	Path to the localized package.
-	 * @param OutSource		Path to the source package.
+	 * @param OutSource		Path to the source package. If the conversion failed, OutSource won't change.
 	 *
-	 * @returns True if the conversion happened, false otherwise
+	 * @returns True if the conversion happened, false otherwise (it will fail it the input was not a localized variant)
 	 */
-	static COREUOBJECT_API bool ConvertLocalizedToSource(const FString& InLocalized, FString& OutSource);
+	static COREUOBJECT_API bool ConvertLocalizedToSource(const FStringView& InLocalized, FString& OutSource);
 
 	/**
+	 * Convert a package path to its source version (by removing /L10N/<code> from the package path, if present)
+	 * Note: This does not check whether the source package exists
+	 *
+	 * @param InPath		Path to a package. It can be a source asset already or a localized variant.
+	 * @param OutSource		Path to the source package. It cannot fail. If no /L10N/<code> were found, it is already a Source Asset.
+	 */
+	static COREUOBJECT_API void ConvertToSource(const FStringView& InPath, FString& OutSource);
+
+	/** 
 	 * Converts a source version of a package path to the localized version for the given culture (by adding /L10N/<code> to the package path)
 	 * Note: This does not check whether the source package exists
 	 *
 	 * @param InSource		Path to the source package.
 	 * @param InCulture		Culture code to use.
-	 * @param OutLocalized	Path to the localized package.
+	 * @param OutLocalized	Path to the localized package. If the conversion failed, OutLocalized won't change.
+	 *
+	 * @returns True if the conversion happened, false otherwise
+	 */
+	static COREUOBJECT_API bool ConvertSourceToLocalized(const FStringView& InSource, const FStringView& InCulture, FString& OutLocalized);
+
+	/**
+	 * Converts a source version of a package path to the regex localized version (by adding "/L10N/" + "*" to the package path)
+	 * Note: This does not check whether the source package exists
+	 *
+	 * @param InSource		Path to the source package.
+	 * @param OutLocalized	Path to the regex localized package. If the conversion failed, OutLocalized won't change.
 	 *
 	 * @returns True if the conversion happened, false otherwise
 	 */
-	static COREUOBJECT_API bool ConvertSourceToLocalized(const FString& InSource, const FString& InCulture, FString& OutLocalized);
+	static COREUOBJECT_API bool ConvertSourceToRegexLocalized(const FStringView& InSource, FString& OutLocalized);
 
 	/**
 	 * Given a package path, get the localized root package for the given culture (eg, if given "/Game/MyFolder/MyAsset" and a culture of "fr", this would return "/Game/L10N/fr")
@@ -41,4 +61,25 @@ struct FPackageLocalizationUtil
 	 * @returns True if the conversion happened, false otherwise
 	 */
 	static COREUOBJECT_API bool GetLocalizedRoot(const FString& InPath, const FString& InCulture, FString& OutLocalized);
+
+	/**
+	 * Extract a culture code from a localized version of a package path (by finding /L10N/<code> from the package path, if present)
+	 * Note: This does not check whether the source package exists
+	 *
+	 * @param InLocalized	Path to the localized package.
+	 * @param OutCulture	Culture code (found from /L10N/<code>)
+	 *
+	 * @returns True if the conversion happened, false otherwise
+	 */
+	static COREUOBJECT_API bool ExtractCultureFromLocalized(const FStringView& InLocalized, FString& OutCulture);
+
+	/**
+	 * Given a package path, returns all localized variants absolute paths found on disk
+	 * Note: There is a similar function to check in Revision Control (if not on disk) 
+	 * (see USourceControlHelpers::GetLocalizedVariantsAbsolutePaths)
+	 *
+	 * @param InPackageName				Package path to use.
+	 * @param OutLocalizedAbsolutePaths	Absolute paths of localized variants founds
+	 */
+	static COREUOBJECT_API void GetLocalizedVariantsAbsolutePaths(const FStringView& InSource, TArray<FString>& OutLocalizedAbsolutePaths);
 };
