Hotfix for implicit casting from ints to floats

#rnx

diff --git a/Engine/Source/Developer/AssetTools/Private/LocalizedAssetTools.cpp b/Engine/Source/Developer/AssetTools/Private/LocalizedAssetTools.cpp
--- a/Engine/Source/Developer/AssetTools/Private/LocalizedAssetTools.cpp
+++ b/Engine/Source/Developer/AssetTools/Private/LocalizedAssetTools.cpp
@@ -67,7 +67,7 @@ ELocalizedAssetsOnDiskResult FLocalizedAssetTools::GetLocalizedVariantsOnDisk(co
 	}
 
 	UAssetDefinitionRegistry* AssetDefinitionRegistry = UAssetDefinitionRegistry::Get();
-	float ProgressStep = 1.0f / InPackages.Num();
+	float ProgressStep = 1.0f / (float)InPackages.Num();
 	for (const FName& OriginalAssetName : InPackages)
 	{
 		GettingLocalizedVariantsOnDiskSlowTask.EnterProgressFrame(ProgressStep);
@@ -145,7 +145,7 @@ ELocalizedAssetsInSCCResult FLocalizedAssetTools::GetLocalizedVariantsInRevision
 	bool bOutRevisionControlWasNeeded = !GetLocalizedVariantsDepotPaths(PackagesAsString, LocalizedVariantsInRevisionControl);
 
 	// Fill a proper structure with the results
-	float ProgressStep = 0.03f / LocalizedVariantsInRevisionControl.Num();
+	float ProgressStep = 0.03f / (float)LocalizedVariantsInRevisionControl.Num();
 	for (const FString& LocalizedVariantInRevisionControl : LocalizedVariantsInRevisionControl)
 	{
 		GetLocalizedVariantsInRevisionControlSlowTask.EnterProgressFrame(ProgressStep);
@@ -160,7 +160,7 @@ ELocalizedAssetsInSCCResult FLocalizedAssetTools::GetLocalizedVariantsInRevision
 	// Don't forget to return the information on the packages that found nothing in Revision Control
 	if (OutPackagesNotFound != nullptr)
 	{
-		ProgressStep = 0.02f / InPackages.Num();
+		ProgressStep = 0.02f / (float)InPackages.Num();
 		for (const FName& PackageName : InPackages)
 		{
 			GetLocalizedVariantsInRevisionControlSlowTask.EnterProgressFrame(ProgressStep);
