Hotfix for implicit casting ints to floats

#rnx

diff --git a/Engine/Source/Developer/AssetTools/Private/AssetRenameManager.cpp b/Engine/Source/Developer/AssetTools/Private/AssetRenameManager.cpp
--- a/Engine/Source/Developer/AssetTools/Private/AssetRenameManager.cpp
+++ b/Engine/Source/Developer/AssetTools/Private/AssetRenameManager.cpp
@@ -1693,7 +1693,7 @@ void FAssetRenameManager::PerformAssetRename(TArray<FAssetRenameDataWithReferenc
 
 	TArray<UPackage*> PackagesToSave = ReferencingPackagesToSave;
 	TArray<UPackage*> PotentialPackagesToDelete;
-	float ProgressStep = 2.0f / AssetsToRename.Num();
+	float ProgressStep = 2.0f / (float)AssetsToRename.Num();
 	for (int32 AssetIdx = 0; AssetIdx < AssetsToRename.Num(); ++AssetIdx)
 	{
 		RenamingSlowTask.EnterProgressFrame(ProgressStep);
@@ -1903,7 +1903,7 @@ void FAssetRenameManager::PerformAssetRename(TArray<FAssetRenameDataWithReferenc
 	FScopedSlowTask BuildingRelationshipSlowTask(5.0f, LOCTEXT("BranchingRenamedAssetsSlowTask", "Building a relationship between the renamed files in Revision Control..."));
 
 	// Now branch the files in source control if possible
-	ProgressStep = 2.9f / AssetsToRename.Num();
+	ProgressStep = 2.9f / (float)AssetsToRename.Num();
 	for (const FAssetRenameDataWithReferencers& RenameData : AssetsToRename)
 	{
 		BuildingRelationshipSlowTask.EnterProgressFrame(ProgressStep);
