Hotfix for cast precision lost

#rnx

diff --git a/Engine/Source/Developer/AssetTools/Private/AssetRenameManager.cpp b/Engine/Source/Developer/AssetTools/Private/AssetRenameManager.cpp
--- a/Engine/Source/Developer/AssetTools/Private/AssetRenameManager.cpp
+++ b/Engine/Source/Developer/AssetTools/Private/AssetRenameManager.cpp
@@ -1289,7 +1289,7 @@ void FAssetRenameManager::SetupPublicAssets(TArray<FAssetRenameDataWithReference
 		return;
 	}
 
-	FScopedSlowTask SlowTask((float)AssetsToRename.Num(), LOCTEXT("SetupPublicAssets", "Setting up public assets..."));
+	FScopedSlowTask SlowTask(static_cast<float>(AssetsToRename.Num()), LOCTEXT("SetupPublicAssets", "Setting up public assets..."));
 	SlowTask.MakeDialog();
 
 	IAssetTools& AssetTools = IAssetTools::Get();
@@ -1693,7 +1693,7 @@ void FAssetRenameManager::PerformAssetRename(TArray<FAssetRenameDataWithReferenc
 
 	TArray<UPackage*> PackagesToSave = ReferencingPackagesToSave;
 	TArray<UPackage*> PotentialPackagesToDelete;
-	float ProgressStep = 2.0f / (float)AssetsToRename.Num();
+	float ProgressStep = 2.0f / static_cast<float>(AssetsToRename.Num());
 	for (int32 AssetIdx = 0; AssetIdx < AssetsToRename.Num(); ++AssetIdx)
 	{
 		RenamingSlowTask.EnterProgressFrame(ProgressStep);
@@ -1903,7 +1903,7 @@ void FAssetRenameManager::PerformAssetRename(TArray<FAssetRenameDataWithReferenc
 	FScopedSlowTask BuildingRelationshipSlowTask(5.0f, LOCTEXT("BranchingRenamedAssetsSlowTask", "Building a relationship between the renamed files in Revision Control..."));
 
 	// Now branch the files in source control if possible
-	ProgressStep = 2.9f / (float)AssetsToRename.Num();
+	ProgressStep = 2.9f / static_cast<float>(AssetsToRename.Num());
 	for (const FAssetRenameDataWithReferencers& RenameData : AssetsToRename)
 	{
 		BuildingRelationshipSlowTask.EnterProgressFrame(ProgressStep);
diff --git a/Engine/Source/Developer/AssetTools/Private/LocalizedAssetTools.cpp b/Engine/Source/Developer/AssetTools/Private/LocalizedAssetTools.cpp
--- a/Engine/Source/Developer/AssetTools/Private/LocalizedAssetTools.cpp
+++ b/Engine/Source/Developer/AssetTools/Private/LocalizedAssetTools.cpp
@@ -67,7 +67,7 @@ ELocalizedAssetsOnDiskResult FLocalizedAssetTools::GetLocalizedVariantsOnDisk(co
 	}
 
 	UAssetDefinitionRegistry* AssetDefinitionRegistry = UAssetDefinitionRegistry::Get();
-	float ProgressStep = 1.0f / (float)InPackages.Num();
+	float ProgressStep = 1.0f / static_cast<float>(InPackages.Num());
 	for (const FName& OriginalAssetName : InPackages)
 	{
 		GettingLocalizedVariantsOnDiskSlowTask.EnterProgressFrame(ProgressStep);
@@ -145,7 +145,7 @@ ELocalizedAssetsInSCCResult FLocalizedAssetTools::GetLocalizedVariantsInRevision
 	bool bOutRevisionControlWasNeeded = !GetLocalizedVariantsDepotPaths(PackagesAsString, LocalizedVariantsInRevisionControl);
 
 	// Fill a proper structure with the results
-	float ProgressStep = 0.03f / (float)LocalizedVariantsInRevisionControl.Num();
+	float ProgressStep = 0.03f / static_cast<float>(LocalizedVariantsInRevisionControl.Num());
 	for (const FString& LocalizedVariantInRevisionControl : LocalizedVariantsInRevisionControl)
 	{
 		GetLocalizedVariantsInRevisionControlSlowTask.EnterProgressFrame(ProgressStep);
@@ -160,7 +160,7 @@ ELocalizedAssetsInSCCResult FLocalizedAssetTools::GetLocalizedVariantsInRevision
 	// Don't forget to return the information on the packages that found nothing in Revision Control
 	if (OutPackagesNotFound != nullptr)
 	{
-		ProgressStep = 0.02f / (float)InPackages.Num();
+		ProgressStep = 0.02f / static_cast<float>(InPackages.Num());
 		for (const FName& PackageName : InPackages)
 		{
 			GetLocalizedVariantsInRevisionControlSlowTask.EnterProgressFrame(ProgressStep);
