Null check before checking if NewObject is in the temporary package

[RN] Fixed crash when diffing assets in inactive plugins

diff --git a/Engine/Source/Developer/AssetTools/Private/AssetTypeActions/AssetTypeActions_DataAsset.cpp b/Engine/Source/Developer/AssetTools/Private/AssetTypeActions/AssetTypeActions_DataAsset.cpp
--- a/Engine/Source/Developer/AssetTools/Private/AssetTypeActions/AssetTypeActions_DataAsset.cpp
+++ b/Engine/Source/Developer/AssetTools/Private/AssetTypeActions/AssetTypeActions_DataAsset.cpp
@@ -75,7 +75,7 @@ void FAssetTypeActions_DataAsset::PerformAssetDiff(UObject* OldAsset, UObject* N
 {
 	const TSharedRef<SDetailsDiff> DetailsDiff = SDetailsDiff::CreateDiffWindow(OldAsset, NewAsset, OldRevision, NewRevision, GetSupportedClass());
 	// allow users to edit NewAsset if it's a local asset
-	if (!FPackageName::IsTempPackage(NewAsset->GetPackage()->GetName()))
+	if (NewAsset && !FPackageName::IsTempPackage(NewAsset->GetPackage()->GetName()))
 	{
 		DetailsDiff->SetOutputObject(NewAsset);
 	}
