Added another dialog option in the new LocalizedAssetTools

diff --git a/Engine/Source/Developer/AssetTools/Private/LocalizedAssetTools.cpp b/Engine/Source/Developer/AssetTools/Private/LocalizedAssetTools.cpp
--- a/Engine/Source/Developer/AssetTools/Private/LocalizedAssetTools.cpp
+++ b/Engine/Source/Developer/AssetTools/Private/LocalizedAssetTools.cpp
@@ -216,11 +216,16 @@ void FLocalizedAssetTools::OpenRevisionControlRequiredDialog() const
 
 void FLocalizedAssetTools::OpenFilesInRevisionControlRequiredDialog(const TArray<FText>& FileList) const
 {
-	SFileListReportDialog::OpenDialog(FilesNeedToBeOnDiskWarningText,
+	OpenLocalizedVariantsListMessageDialog(FilesNeedToBeOnDiskWarningText,
 		LOCTEXT("FilesToSyncDialogHeader", "The following assets were found only in Revision Control. They need to be on your disk to be renamed."),
 		FileList);
 }
 
+void FLocalizedAssetTools::OpenLocalizedVariantsListMessageDialog(const FText& Header, const FText& Message, const TArray<FText>& FileList) const
+{
+	SFileListReportDialog::OpenDialog(Header, Message, FileList);
+}
+
 const FText& FLocalizedAssetTools::GetRevisionControlIsNotAvailableWarningText() const
 {
 	return RevisionControlIsNotAvailableWarningText;
diff --git a/Engine/Source/Developer/AssetTools/Private/LocalizedAssetTools.h b/Engine/Source/Developer/AssetTools/Private/LocalizedAssetTools.h
--- a/Engine/Source/Developer/AssetTools/Private/LocalizedAssetTools.h
+++ b/Engine/Source/Developer/AssetTools/Private/LocalizedAssetTools.h
@@ -27,6 +27,8 @@ public:
 
 	virtual void OpenFilesInRevisionControlRequiredDialog(const TArray<FText>& FileList) const override;
 
+	virtual void OpenLocalizedVariantsListMessageDialog(const FText& Header, const FText& Message, const TArray<FText>& FileList) const override;
+
 	virtual const FText& GetRevisionControlIsNotAvailableWarningText() const override;
 
 	virtual const FText& GetFilesNeedToBeOnDiskWarningText() const override;
diff --git a/Engine/Source/Developer/AssetTools/Public/ILocalizedAssetTools.h b/Engine/Source/Developer/AssetTools/Public/ILocalizedAssetTools.h
--- a/Engine/Source/Developer/AssetTools/Public/ILocalizedAssetTools.h
+++ b/Engine/Source/Developer/AssetTools/Public/ILocalizedAssetTools.h
@@ -50,6 +50,9 @@ public:
 	/** Open a dialog to required files in Revision Control to be on disk */
 	virtual void OpenFilesInRevisionControlRequiredDialog(const TArray<FText>& FileList) const = 0;
 
+	/** Open a dialog with a file list (localized variants) with a custom message */
+	virtual void OpenLocalizedVariantsListMessageDialog(const FText& Header, const FText& Message, const TArray<FText>& FileList) const = 0;
+
 	/** Get shared text to warn that some files need to be on disk (not only in Revision Control) */
 	virtual const FText& GetFilesNeedToBeOnDiskWarningText() const = 0;
 
