Editor - When creating assets call new delegate to customize when to set package flag NotExternallyReferenceable
#rnx

diff --git a/Engine/Source/Developer/AssetTools/Private/AssetTools.cpp b/Engine/Source/Developer/AssetTools/Private/AssetTools.cpp
--- a/Engine/Source/Developer/AssetTools/Private/AssetTools.cpp
+++ b/Engine/Source/Developer/AssetTools/Private/AssetTools.cpp
@@ -1751,7 +1751,7 @@ UObject* UAssetToolsImpl::CreateAsset(const FString& AssetName, const FString& P
 	if( NewObj )
 	{
 
-		Pkg->SetIsExternallyReferenceable(CreateAssetsAsExternallyReferenceable);
+		Pkg->SetIsExternallyReferenceable(ShouldCreateAssetsAsExternallyReferenceableForPath(PackageName));
 
 		// Notify the asset registry
 		FAssetRegistryModule::AssetCreated(NewObj);
@@ -1991,6 +1991,25 @@ bool UAssetToolsImpl::GetCreateAssetsAsExternallyReferenceable()
 	return CreateAssetsAsExternallyReferenceable;
 }
 
+bool UAssetToolsImpl::ShouldCreateAssetsAsExternallyReferenceableForPath(const FStringView AssetPath) const
+{
+	// Allow other systems to control the result
+	TOptional<bool> bExternallyReferenceable;
+	CreateAssetsAsExternallyReferenceableForPathDelegate.Broadcast(AssetPath, bExternallyReferenceable);
+	if (bExternallyReferenceable.IsSet())
+	{
+		return bExternallyReferenceable.GetValue();
+	}
+
+	// If no other systems set a result, then use the global setting applied to all assets
+	return CreateAssetsAsExternallyReferenceable;
+}
+
+UE::AssetTools::FShouldCreateAssetsAsExternallyReferenceableForPath& UAssetToolsImpl::GetOnShouldCreateAssetsAsExternallyReferenceableForPath()
+{
+	return CreateAssetsAsExternallyReferenceableForPathDelegate;
+}
+
 void UAssetToolsImpl::GenerateAdvancedCopyDestinations(FAdvancedCopyParams& InParams, const TArray<FName>& InPackageNamesToCopy, const UAdvancedCopyCustomization* CopyCustomization, TMap<FString, FString>& OutPackagesAndDestinations) const
 {
 	IDesktopPlatform* DesktopPlatform = FDesktopPlatformModule::Get();
@@ -4517,7 +4536,7 @@ TArray<UObject*> UAssetToolsImpl::ImportAssetsInternal(const TArray<FString>& Fi
 				continue;
 			}
 
-			Pkg->SetIsExternallyReferenceable(CreateAssetsAsExternallyReferenceable);
+			Pkg->SetIsExternallyReferenceable(ShouldCreateAssetsAsExternallyReferenceableForPath(PackageName));
 
 			Factory->SetAutomatedAssetImportData(Params.ImportData);
 			Factory->SetAssetImportTask(Params.AssetImportTask);
diff --git a/Engine/Source/Developer/AssetTools/Private/AssetTools.h b/Engine/Source/Developer/AssetTools/Private/AssetTools.h
--- a/Engine/Source/Developer/AssetTools/Private/AssetTools.h
+++ b/Engine/Source/Developer/AssetTools/Private/AssetTools.h
@@ -92,6 +92,8 @@ public:
 	virtual UObject* DuplicateAssetWithDialogAndTitle(const FString& AssetName, const FString& PackagePath, UObject* OriginalObject, FText DialogTitle) override;
 	virtual void SetCreateAssetsAsExternallyReferenceable(bool bValue) override;
 	virtual bool GetCreateAssetsAsExternallyReferenceable() override;
+	virtual bool ShouldCreateAssetsAsExternallyReferenceableForPath(const FStringView AssetPath) const override;
+	virtual UE::AssetTools::FShouldCreateAssetsAsExternallyReferenceableForPath& GetOnShouldCreateAssetsAsExternallyReferenceableForPath() override;
 	virtual bool IsDiscoveringAssetsInProgress() const override;
 	virtual void OpenDiscoveringAssetsDialog(const FOnAssetsDiscovered& InOnAssetsDiscovered) override;
 	virtual bool RenameAssets(const TArray<FAssetRenameData>& AssetsAndNames) override;
@@ -277,6 +279,7 @@ private:
 	TArray<FString> SubContentDenyListPaths;
 
 	bool CreateAssetsAsExternallyReferenceable;
+	UE::AssetTools::FShouldCreateAssetsAsExternallyReferenceableForPath CreateAssetsAsExternallyReferenceableForPathDelegate;
 
 	TMap<FName, FIsNameAllowed> IsNameAllowedDelegates;
 	TMap<FName, FSanitizeName> SanitizeNameDelegates;
diff --git a/Engine/Source/Developer/AssetTools/Public/IAssetTools.h b/Engine/Source/Developer/AssetTools/Public/IAssetTools.h
--- a/Engine/Source/Developer/AssetTools/Public/IAssetTools.h
+++ b/Engine/Source/Developer/AssetTools/Public/IAssetTools.h
@@ -38,6 +38,8 @@ namespace UE::AssetTools
 	DECLARE_DELEGATE_RetVal_OneParam(bool, FCanMigrateAsset, FName);
 
 	DECLARE_DELEGATE_RetVal_OneParam(bool, FCanAssetBePublic, FStringView /*AssetPath*/);
+
+	DECLARE_MULTICAST_DELEGATE_TwoParams(FShouldCreateAssetsAsExternallyReferenceableForPath, FStringView /*AssetPath*/, TOptional<bool>& /*ExternallyReferenceable*/);
 }
 
 UENUM()
@@ -384,9 +386,15 @@ public:
 	/** Controls whether or not newly created assets are made externally referneceable or not */
 	virtual void SetCreateAssetsAsExternallyReferenceable(bool bValue) = 0;
 
-	/** Gets whether assets are being made externally referenceable or not */
+	/** Gets whether all assets being made at any location are externally referenceable or not */
 	virtual bool GetCreateAssetsAsExternallyReferenceable() = 0;
 
+	/** Gets whether assets being made at this location are externally referenceable or not */
+	virtual bool ShouldCreateAssetsAsExternallyReferenceableForPath(const FStringView AssetPath) const = 0;
+
+	/** Delegate to control whether assets being made at a specific location are externally referenceable or not */
+	virtual UE::AssetTools::FShouldCreateAssetsAsExternallyReferenceableForPath& GetOnShouldCreateAssetsAsExternallyReferenceableForPath() = 0;
+
 	/** Gets whether assets registry is still loading assets or not */
 	virtual bool IsDiscoveringAssetsInProgress() const = 0;
 
diff --git a/Engine/Source/Editor/GameProjectGeneration/Private/SNewClassDialog.cpp b/Engine/Source/Editor/GameProjectGeneration/Private/SNewClassDialog.cpp
--- a/Engine/Source/Editor/GameProjectGeneration/Private/SNewClassDialog.cpp
+++ b/Engine/Source/Editor/GameProjectGeneration/Private/SNewClassDialog.cpp
@@ -1154,7 +1154,7 @@ void SNewClassDialog::FinishClicked()
 				if (NewBP)
 				{
 					// Set the default "IsExternallyReferenceable" state
-					Package->SetIsExternallyReferenceable(IAssetTools::Get().GetCreateAssetsAsExternallyReferenceable());
+					Package->SetIsExternallyReferenceable(IAssetTools::Get().ShouldCreateAssetsAsExternallyReferenceableForPath(PackagePath));
 
 					// Notify the asset registry
 					FAssetRegistryModule::AssetCreated(NewBP);
