Adding an API to IAssetTools to provide a callback for valildating whether a new asset can be created at /Path/To/New/Package. It differs from existing APIs by providing both the path and the package name to the callback as an FSoftObjectPath.

[RN] Added an API to IAssetTools to provide a callback for valildating whether a new asset can be created at /Path/To/New/Package.

diff --git a/Engine/Source/Developer/AssetTools/Private/AssetTools.cpp b/Engine/Source/Developer/AssetTools/Private/AssetTools.cpp
--- a/Engine/Source/Developer/AssetTools/Private/AssetTools.cpp
+++ b/Engine/Source/Developer/AssetTools/Private/AssetTools.cpp
@@ -5106,6 +5106,13 @@ bool UAssetToolsImpl::CanCreateAsset(const FString& AssetName, const FString& Pa
 		return false;
 	}
 
+	// Check for delegates that disallow the particular path and name
+	if (!IsObjectPathAllowed(PackageName, &Reason))
+	{
+		FMessageDialog::Open(EAppMsgType::Ok, Reason);
+		return false;
+	}
+
 	// Find (or create!) the desired package for this object
 	UPackage* Pkg = CreatePackage(*PackageName);
 
@@ -6836,6 +6843,28 @@ void UAssetToolsImpl::UnregisterIsNameAllowedDelegate(const FName OwnerName)
 	IsNameAllowedDelegates.Remove(OwnerName);
 }
 
+bool UAssetToolsImpl::IsObjectPathAllowed(const FString& FullPathWithName, FText* OutErrorMessage /* = nullptr */) const
+{
+	for (const TPair<FName, FIsObjectPathAllowed>& DelegatePair : IsObjectPathAllowedDelegates)
+	{
+		if (!DelegatePair.Value.Execute(FullPathWithName, OutErrorMessage))
+		{
+			return false;
+		}
+	}
+	return true;
+}
+
+void UAssetToolsImpl::RegisterIsObjectPathAllowedDelegate(const FName OwnerName, FIsObjectPathAllowed Delegate)
+{
+	IsObjectPathAllowedDelegates.Add(OwnerName, Delegate);
+}
+
+void UAssetToolsImpl::UnregisterIsObjectPathAllowedDelegate(const FName OwnerName)
+{
+	IsObjectPathAllowedDelegates.Remove(OwnerName);
+}
+
 bool UAssetToolsImpl::SanitizeName(FString& NameToSanitize)
 {
 	//Call the delegates
diff --git a/Engine/Source/Developer/AssetTools/Private/AssetTools.h b/Engine/Source/Developer/AssetTools/Private/AssetTools.h
--- a/Engine/Source/Developer/AssetTools/Private/AssetTools.h
+++ b/Engine/Source/Developer/AssetTools/Private/AssetTools.h
@@ -153,6 +153,9 @@ public:
 	virtual bool IsNameAllowed(const FString& Name, FText* OutErrorMessage = nullptr) const override;
 	virtual void RegisterIsNameAllowedDelegate(const FName OwnerName, FIsNameAllowed Delegate) override;
 	virtual void UnregisterIsNameAllowedDelegate(const FName OwnerName) override;
+	virtual bool IsObjectPathAllowed(const FString& FullPathWithName, FText* OutErrorMessage = nullptr) const override;
+	virtual void RegisterIsObjectPathAllowedDelegate(const FName OwnerName, FIsObjectPathAllowed Delegate) override;
+	virtual void UnregisterIsObjectPathAllowedDelegate(const FName OwnerName) override;
 	virtual bool SanitizeName(FString& NameToSanitize) override;
 	virtual void RegisterSanitizeNameDelegate(const FName OwnerName, FSanitizeName Delegate) override;
 	virtual void UnregisterSanitizeNameDelegate(const FName OwnerName) override;
@@ -282,6 +285,7 @@ private:
 	UE::AssetTools::FShouldCreateAssetsAsExternallyReferenceableForPath CreateAssetsAsExternallyReferenceableForPathDelegate;
 
 	TMap<FName, FIsNameAllowed> IsNameAllowedDelegates;
+	TMap<FName, FIsObjectPathAllowed> IsObjectPathAllowedDelegates;
 	TMap<FName, FSanitizeName> SanitizeNameDelegates;
 
 	UE::AssetTools::FOnPackageMigration OnPackageMigration;
diff --git a/Engine/Source/Developer/AssetTools/Public/IAssetTools.h b/Engine/Source/Developer/AssetTools/Public/IAssetTools.h
--- a/Engine/Source/Developer/AssetTools/Public/IAssetTools.h
+++ b/Engine/Source/Developer/AssetTools/Public/IAssetTools.h
@@ -151,6 +151,7 @@ DECLARE_DYNAMIC_DELEGATE_TwoParams(FAdvancedCopyCompletedEvent, bool, bSuccess,
 
 DECLARE_MULTICAST_DELEGATE_OneParam(FAssetPostRenameEvent, const TArray<FAssetRenameData>&);
 DECLARE_DELEGATE_RetVal_TwoParams(bool, FIsNameAllowed, const FString& /*Name*/, FText* /*OutErrorMessage*/);
+DECLARE_DELEGATE_RetVal_TwoParams(bool, FIsObjectPathAllowed, const FString& /*FullPathAndName*/, FText* /*OutErrorMessage*/);
 DECLARE_DELEGATE_OneParam(FSanitizeName, FString& /*NameToSanitize*/);
 
 struct FAdvancedAssetCategory
@@ -667,6 +668,13 @@ public:
 	/** Remove a previously-set global name filter */
 	virtual void UnregisterIsNameAllowedDelegate(const FName OwnerName) = 0;
 
+	/** Returns true if IsObjectPathAllowedDelegate is not set, or if the name passes the filter function */
+	virtual bool IsObjectPathAllowed(const FString& FullObjectPathWithName, FText* OutErrorMessage = nullptr) const = 0;
+	/** Allows setting of a global name filter that is applied to assets and folders */
+	virtual void RegisterIsObjectPathAllowedDelegate(const FName OwnerName, FIsObjectPathAllowed Delegate) = 0;
+	/** Remove a previously-set global name filter */
+	virtual void UnregisterIsObjectPathAllowedDelegate(const FName OwnerName) = 0;
+
 	/** Sanitize the name by calling all registered delegates. */
 	virtual bool SanitizeName(FString& NameToSanitize) = 0;
 	/** Allows setting of a global name filter that is applied to assets and folders */
