Remove unnecessary check on full object path length
#rnx

diff --git a/Engine/Source/Developer/AssetTools/Private/AssetViewUtils.cpp b/Engine/Source/Developer/AssetTools/Private/AssetViewUtils.cpp
--- a/Engine/Source/Developer/AssetTools/Private/AssetViewUtils.cpp
+++ b/Engine/Source/Developer/AssetTools/Private/AssetViewUtils.cpp
@@ -1437,15 +1437,6 @@ bool AssetViewUtils::IsValidObjectPathForCreate(const FString& ObjectPath, const
 	}
 
 	const FString FullPath = FPaths::ConvertRelativePathToFull(RelativePathFilename);	// path to file on disk
-	if (ObjectPath.Len() > (FPlatformMisc::GetMaxPathLength() - MAX_CLASS_NAME_LENGTH))
-	{
-		// The full path for the asset is too long
-		OutErrorMessage = FText::Format(LOCTEXT("ObjectPathTooLong", "The object path for the asset is too long, the maximum is '{0}'. \nPlease choose a shorter name for the asset or create it in a shallower folder structure."),
-			FText::AsNumber((FPlatformMisc::GetMaxPathLength() - MAX_CLASS_NAME_LENGTH)));
-		// Return false to indicate that the user should enter a new name
-		return false;
-	}
-		
 	if (FullPath.Len() > CVarMaxFullPathLength->GetValueOnGameThread() )
 	{
 		// The full path for the asset is too long
