Changes FAssetHeaderPatcherInner to leverage the CoreRedirect system for resolving FNames in a package header, and string paths when context about the string itself is well known (i.e. we know the string data is a stringified ObjectPath/PackagePath, or ObjectName are deducible). We still make use of string scanning replacement mechanism but those are now constrained to only patching strings content specifically (e.g. Tag data)

- Added FAssetHeaderPatcher::FContext for storing data on how to perform a patch (which redirects to follow, string -> string replacement map). Allows for an extensibility point for encapsulating special case patching so we don't need to  rely on string->string mappings exclusively
- Added `virtual FArchive& operator<<(FTopLevelAssetPath& Value)` to `FActorDescArchive` to allow for patching of FTopLevelAssetPath directly. Without this overload the FNames within a FTopLevelAssetPath would be patched one at a time losing all context about what the FNames are.
- Removed IAssetTools::GetAdditionalPatchCopyMappings, GetPatchCopyMappingsForRootRename, and PatchCopyPackageFile as these methods were internal and using the patcher directly should be preferred.
- Added some unit tests for the FAssetHeaderPatcherInner implementation works to codify the assumptions and expectations of the patcher implementation

New patching changes are left disabled by default

#rnx

diff --git a/Engine/Source/Developer/AssetTools/Internal/AssetHeaderPatcher.h b/Engine/Source/Developer/AssetTools/Internal/AssetHeaderPatcher.h
new file mode 100644
--- /dev/null
+++ b/Engine/Source/Developer/AssetTools/Internal/AssetHeaderPatcher.h
@@ -0,0 +1,222 @@
+// Copyright Epic Games, Inc. All Rights Reserved.
+
+#pragma once
+
+#include "Tasks/Task.h"
+#include "Containers/ContainersFwd.h"
+#include "UObject/CoreRedirects.h"
+
+/**
+ * Delegate called on when patch operation completes
+ * @param	SrcFilePath		Path of file being read for patching
+ * @param	DstFilePath		Path of file being written to after patching
+ */
+DECLARE_DELEGATE_TwoParams(FAssetHeaderPatcherCompletionDelegate, const FString& /*SrcFilePath*/, const FString& /*DstFilePath*/)
+
+UE_INTERNAL struct ASSETTOOLS_API FAssetHeaderPatcher
+{
+	enum class EResult
+	{
+		NotStarted,
+		Cancelled,
+		InProgress,
+		Success,
+		ErrorFailedToLoadSourceAsset,
+		ErrorFailedToDeserializeSourceAsset,
+		ErrorUnexpectedSectionOrder,
+		ErrorBadOffset,
+		ErrorUnkownSection,
+		ErrorFailedToOpenDestinationFile,
+		ErrorFailedToWriteToDestinationFile,
+		ErrorEmptyRequireSection,
+	};
+
+	UE_INTERNAL struct ASSETTOOLS_API FContext
+	{
+		FContext() = default;
+
+		/**
+		 * Context used for patching. Contains all information for how object and package references
+		 * will be changed as part of patching. 
+		 * *
+		 * When bInGatherDependentPackages is true, the provided long package name (/Root/Folder/Package) to 
+		 * destination long package name mapping will be used to find any dependent packages that must also 
+		 * be patched due to internal references. The mapping provided in InSrcAndDstPackagePaths will be used 
+		 * to determine the filepath on disk to write when patching.
+		 *
+		 * @param InSrcAndDstPackagePaths Map of all long package names (/Root/Folder/Package) to be patched and to which new name they should be patched to.
+		 * @param bInGatherDependentPackages If true (default), upon creating the context GatherDependentPackages() will be called.
+		 **/
+		FContext(const TMap<FString, FString>& InSrcAndDstPackagePaths, const bool bInGatherDependentPackages = true);
+
+		/**
+		 * Context used for patching. Contains all information for how object and package references
+		 * will be changed as part of patching. 
+		 * 
+		 * When patching, package paths to patch will be deduced by the filepath mappings provided in InSrcAndDstFilePaths. All assets
+		 * under InSrcRoot will be written as package paths under a mountpoint located at InSrcBaseDir.
+		 * 
+		 * e.g. Path "C:/User/Repo/Project/Content/Skeletons/Player.uasset" -> "/InSrcRoot/Skeletons/Player" when InSrcBaseDir=C:/User/Repo/Project (/Content is assumed internally)
+		 * 
+		 * @param InSrcRoot The root mount point for assets to be patched
+		 * @param InDstRoot The new root mount point for patched assets to be placed under
+		 * @param InSrcBaseDir Path to the directory holding the /Content/ directory for assets to patch
+		 * @param InSrcAndDstFilePaths Map of filepaths for files to be patched and where to write the patched version to
+		 * @param InMountPointReplacements Map of root mountpoints (name only, no "/" prefix or suffix) to replace when patching
+		 **/
+		FContext(const FString& InSrcRoot, const FString& InDstRoot, const FString& InSrcBaseDir, const TMap<FString, FString>& InSrcAndDstFilePaths, const TMap<FString, FString>& InMountPointReplacements);
+
+		/*
+		* Returns the mapping of source long package names to destination package paths used when patching.
+		* This mapping may include more packages than initially supplied to the FContext
+		* if GatherDependentPackages has already been called. 
+		* Note, this map can be invalidated by calls to GatherDependentPackages()
+		*/
+		const TMap<FString, FString>& GetLongPackagePathRemapping() const
+		{
+			return PackagePathRenameMap;
+		};
+
+	protected:
+		friend FAssetHeaderPatcher;
+
+		void AddVerseMounts();
+		void GatherDependentPackages();
+		void GenerateFilePathsFromPackagePaths();
+		void GeneratePackagePathsFromFilePaths(const FString& InSrcRoot, const FString& InDstRoot, const FString& InSrcBaseDir);
+		void GenerateAdditionalRemappings();
+	
+		TArray<FString> VerseMountPoints;
+		TMap<FString, FString> PackagePathRenameMap;
+		TMap<FString, FString> FilePathRenameMap;
+
+		// Todo: Make TSet once FCoreRedirect GetTypeHash is implemented
+		TArray<FCoreRedirect> Redirects;
+
+		// String mappings are only used for best-effort replacements. These will be error-prone 
+		// and we should strive for more structured data formats to guard against errors here
+		TMap<FString, FString> StringReplacements;	
+		TMap<FString, FString> StringMountReplacements;
+	};
+
+	FAssetHeaderPatcher() = default;
+	FAssetHeaderPatcher(FContext InContext) : Context(InContext) { Reset(); }
+	
+	/*
+	* Resets the patcher state and sets a new patching context.
+	* It is an error to call while patching is already in progress.
+	*/
+	void SetContext(FContext InContext);
+
+	/*
+	* Schedules the reading of source files determined by the patcher context, as well as the writing of the patched versions
+	* of all source files read. 
+	* 
+	* @param InOutNumFilesToPatch Optional value to know how many files are expected to be read/written during patching
+	* @param InOutNumFilesPatched Optional value used to know how the patcher is progressing (useful for progress bars)
+	*/
+	UE::Tasks::FTask PatchAsync(int32* InOutNumFilesToPatch = nullptr, int32* InOutNumFilesPatched = nullptr);
+	UE::Tasks::FTask PatchAsync(int32* InOutNumFilesToPatch, int32* InOutNumFilesPatched, FAssetHeaderPatcherCompletionDelegate InOnSuccess, FAssetHeaderPatcherCompletionDelegate InOnError);
+
+	/*
+	* Returns the status of any inflight patching operations. In the case of multiple errors, the last seen error will be reported.
+	* Per file error status codes can be returned with GetErrorFiles().
+	*/
+	EResult GetPatchResult() const
+	{
+		return Status;
+	}
+
+	/*
+	* Returns source file -> destination mapping for all files that were patched successfully.
+	*/
+	TMap<FString, FString> GetPatchedFiles() const
+	{
+		if (IsPatching())
+		{
+			return TMap<FString, FString>();
+		}
+
+		return PatchedFiles;
+	}
+
+	/*
+	* Returns true if the patcher encountered errors (even if patching was cancelled)
+	*/
+	bool HasErrors()
+	{
+		FScopeLock Lock(&ErroredFilesLock);
+		return !!ErroredFiles.Num();
+	}
+
+	/*
+	* Returns a map of all files that had an error during patching with an error code to provide context as to the cause of the error.
+	*/
+	TMap<FString, EResult> GetErrorFiles()
+	{
+		FScopeLock Lock(&ErroredFilesLock);
+		return ErroredFiles;
+	}
+
+	/*
+	* Returns true if the patcher is still in theprocess of patching.
+	*/
+	bool IsPatching() const 
+	{ 
+		return !PatchingTask.IsCompleted(); 
+	}
+
+	/*
+	* Cancels an in-flight patching operation. Patching work on individual files that has already started will run to completion
+	* however any files that have not started patching will be skipped. Even after cancelling, one must wait for the patcher to complete
+	* by waiting on the GetPatchingTask() explciitly or until IsPatching returns false.
+	* 
+	* @return true if an in-flight patching operation was cancelled. If no patching operation is underway, returns false.
+	*/
+	bool CancelPatching() 
+	{
+		if (!IsPatching())
+		{
+			return false;
+		}
+
+		bCancelled = true;
+		Status = EResult::Cancelled;
+
+		return true;
+	}
+
+	/*
+	* Returns the task for all patcher work underway. Waiting on this task will guarantee all patch work is completed.
+	*/
+	UE::Tasks::FTask GetPatchingTask() const
+	{
+		return PatchingTask;
+	}
+
+	/**
+	 * Patches object and package references contained within InSrcAsset using the mapping provided to InContext. The 
+	 * patched asset will be written to InDstAsset. 
+	 *
+	 * @param InSrcAsset Long package name (/Root/Folder/Package) to read in to be patched.
+	 * @param InDstAsset Long package name (/Root/Folder/Package) where the patched package will be written to.
+	 * @param InContext Context for how the patching will be performed. Contains all remapping information to the patcher.
+	 * @return Success patching was successful and the InDstAsset package was written. Returns an error status otherwise.
+	 **/
+	static EResult DoPatch(const FString& InSrcAsset, const FString& InDstAsset, const FContext& InContext);
+
+private:
+	void Reset();
+
+	FContext Context;
+
+	TMap<FString, EResult> ErroredFiles;
+	FCriticalSection  ErroredFilesLock;
+
+	TMap<FString, FString> PatchedFiles;
+
+	UE::Tasks::FTask PatchingTask;
+	EResult Status = EResult::NotStarted;
+	std::atomic<bool> bCancelled = false;
+};
+UE_INTERNAL ASSETTOOLS_API FString LexToString(FAssetHeaderPatcher::EResult InResult);
diff --git a/Engine/Source/Developer/AssetTools/Private/AssetHeaderPatcher.cpp b/Engine/Source/Developer/AssetTools/Private/AssetHeaderPatcher.cpp
--- a/Engine/Source/Developer/AssetTools/Private/AssetHeaderPatcher.cpp
+++ b/Engine/Source/Developer/AssetTools/Private/AssetHeaderPatcher.cpp
@@ -2,21 +2,25 @@
 
 #include "AssetHeaderPatcher.h"
 
-#include "Containers/ContainersFwd.h"
-#include "UObject/PackageFileSummary.h"
-#include "UObject/NameTypes.h"
-#include "UObject/ObjectResource.h"
-#include "UObject/Package.h"
-
+#include "Algo/Copy.h"
 #include "AssetRegistry/IAssetRegistry.h"
 #include "AssetRegistry/AssetData.h"
-
+#include "AssetRegistry/PackageReader.h"
+#include "Containers/ContainersFwd.h"
 #include "Internationalization/GatherableTextData.h"
+#include "Misc/Base64.h"
 #include "Misc/EnumerateRange.h"
 #include "Misc/FileHelper.h"
 #include "Misc/PackageName.h"
+#include "Misc/Paths.h"
 #include "Misc/PathViews.h"
 #include "Serialization/LargeMemoryReader.h"
+#include "UObject/CoreRedirects.h"
+#include "UObject/Linker.h"
+#include "UObject/NameTypes.h"
+#include "UObject/ObjectResource.h"
+#include "UObject/Package.h"
+#include "UObject/PackageFileSummary.h"
 #include "WorldPartition/WorldPartitionActorDesc.h"
 #include "WorldPartition/WorldPartitionActorDescUtils.h"
 
@@ -24,173 +28,698 @@ DEFINE_LOG_CATEGORY_STATIC(LogAssetHeaderPatcher, Log, All);
 
 namespace
 {
-	// To override writing of FName's to ensure they have been patched
-	class FNamePatchingWriter final : public FArchiveProxy
+	// If working on header patching, this is very helpful for dumping what is patched and reviewing the files in a folder comparison of your favourite diff program.
+	FString DumpOutputDirectory;
+	static FAutoConsoleVariableRef CVarDumpOutputDirectory(
+		TEXT("AssetHeaderPatcher.DebugDumpDir"),
+		DumpOutputDirectory,
+		TEXT("'Before'/'After' text representations of each package processed during patching will be written out to the provided absolute filesystem path. Useful for comparing what was patched.")
+	);
+
+	// Tag 'Key' names that are generally large blobs of data that can't/shouldn't be patched
+	const TCHAR* TagsToIgnore[] =
+	{
+		TEXT("FiBData")
+	};
+
+	const FStringView InvalidObjectPathCharacters(INVALID_OBJECTPATH_CHARACTERS);
+
+	bool SplitLongPackageName(FStringView LongPackageName, FStringView& PackageRoot, FStringView& PackagePath, FStringView& PackageName)
 	{
-	public:
-		FNamePatchingWriter(FArchive& InAr, const TMap<FNameEntryId, int32>& InNameToIndexMap)
-			: FArchiveProxy(InAr)
-			, NameToIndexMap(InNameToIndexMap)
+		if (LongPackageName.IsEmpty() || LongPackageName[0] != TEXT('/'))
 		{
+			return false;
 		}
 
-		virtual ~FNamePatchingWriter()
+		PackageRoot = FStringView(LongPackageName.GetData() + 1); // + 1 to skip the leading '/'
+		int32 SeparatorPos;
+		if (!PackageRoot.FindChar(TEXT('/'), SeparatorPos))
 		{
+			return false;
 		}
+		PackageRoot.LeftInline(SeparatorPos);
 
-		virtual FArchive& operator<<(FName& Name) override
+		const int32 PackagePathOffset = PackageRoot.Len() + 2; // + 2 for the leading and trailing '/'
+		if (LongPackageName.Len() < PackagePathOffset || !LongPackageName.FindLastChar(TEXT('/'), SeparatorPos))
 		{
-			FNameEntryId EntryId = Name.GetDisplayIndex();
-			const int32* MaybeIndex = NameToIndexMap.Find(EntryId);
+			return false;
+		}
 
-			if (MaybeIndex == nullptr)
-			{
-				ErrorMessage += FString::Printf(TEXT("Cannot serialize FName %s because it is not in the name table for %s\n"), *Name.ToString(), *GetArchiveName());
-				SetCriticalError();
-				return *this;
-			}
+		// May be empty. If the PackageName is off the root there is no PackagePath
+		const int32 PackagePathLen = SeparatorPos - (PackagePathOffset - 1);
+		check(PackagePathLen >= 0);
+		PackagePath = FStringView(LongPackageName.GetData() + PackagePathOffset, PackagePathLen - !!PackagePathLen);
 
-			int32 Index = *MaybeIndex;
-			int32 Number = Name.GetNumber();
+		const int32 PackageNameOffset = PackagePathOffset + PackagePath.Len() + !PackagePath.IsEmpty();
+		PackageName = FStringView(LongPackageName.GetData() + PackageNameOffset, LongPackageName.Len() - PackageNameOffset);
 
-			FArchive& Ar = *this;
-			Ar << Index;
-			Ar << Number;
+		return true;
+	}
 
-			return *this;
+	FStringView Find(const TMap<FString, FString>& Table, FStringView Needle)
+	{
+		uint32 NeedleHash = TMap<FString, FString>::KeyFuncsType::GetKeyHash<FStringView>(Needle);
+		const FString* MaybeNewItem = Table.FindByHash<FStringView>(NeedleHash, Needle);
+		if (MaybeNewItem)
+		{
+			return *MaybeNewItem;
 		}
 
-		const FString& GetErrorMessage() const
+		return {};
+	}
+}
+
+FString LexToString(FAssetHeaderPatcher::EResult InResult)
+{
+	switch (InResult)
+	{
+		case FAssetHeaderPatcher::EResult::NotStarted: return TEXT("Not Started");
+		case FAssetHeaderPatcher::EResult::Cancelled: return TEXT("Cancelled");
+		case FAssetHeaderPatcher::EResult::InProgress: return TEXT("In Progress");
+		case FAssetHeaderPatcher::EResult::Success: return TEXT("Success");
+		case FAssetHeaderPatcher::EResult::ErrorFailedToLoadSourceAsset: return TEXT("Failed to load source asset");
+		case FAssetHeaderPatcher::EResult::ErrorFailedToDeserializeSourceAsset: return TEXT("Failed to deserialize source asset");
+		case FAssetHeaderPatcher::EResult::ErrorUnexpectedSectionOrder: return TEXT("Unexpected section order");
+		case FAssetHeaderPatcher::EResult::ErrorBadOffset: return TEXT("Bad offset");
+		case FAssetHeaderPatcher::EResult::ErrorUnkownSection: return TEXT("Unknown section");
+		case FAssetHeaderPatcher::EResult::ErrorFailedToOpenDestinationFile: return TEXT("Failed to open destination file");
+		case FAssetHeaderPatcher::EResult::ErrorFailedToWriteToDestinationFile: return TEXT("Failed to write to destination file");
+		case FAssetHeaderPatcher::EResult::ErrorEmptyRequireSection: return TEXT("Empty required section");
+		default: return TEXT("Unknown");
+	}
+}
+
+FAssetHeaderPatcher::FContext::FContext(const TMap<FString, FString>& SourceAndDestPackages, const bool bInGatherDependentPackages)
+	: PackagePathRenameMap(SourceAndDestPackages)
+{
+	AddVerseMounts();
+
+	if (bInGatherDependentPackages)
+	{
+		GatherDependentPackages();
+	}
+
+	GenerateFilePathsFromPackagePaths();
+	GenerateAdditionalRemappings();
+}
+
+FAssetHeaderPatcher::FContext::FContext(const FString& InSrcRoot, const FString& InDstRoot, const FString& InSrcBaseDir, const TMap<FString, FString>& InSrcAndDstFilePaths, const TMap<FString, FString>& InMountPointReplacements)
+	: FilePathRenameMap(InSrcAndDstFilePaths)
+	, StringMountReplacements(InMountPointReplacements)
+{
+	AddVerseMounts();
+	GeneratePackagePathsFromFilePaths(InSrcRoot, InDstRoot, InSrcBaseDir);
+	GenerateAdditionalRemappings();
+}
+
+void FAssetHeaderPatcher::FContext::AddVerseMounts()
+{
+	// Todo: Expose this so callers provide this data
+	VerseMountPoints.Add("localhost");
+}
+
+void FAssetHeaderPatcher::FContext::GenerateFilePathsFromPackagePaths()
+{
+	FilePathRenameMap.Reserve(PackagePathRenameMap.Num());
+
+	// Construct all source and destination filenames from our package map
+	for (const TTuple<FString, FString>& Package : PackagePathRenameMap)
+	{
+		const FString& PackageName = Package.Key;
+		const FString& DestPackage = Package.Value;
+		FString SrcFilename;
+
+		// To consider: Allow the caller to provide their own file filter
+		if (FPackageName::IsVersePackage(PackageName))
 		{
-			return ErrorMessage;
+			// Verse packages are not header patchable.
+			// They are also not Packages as far as DoesPackageExist tells me.
+			// But they are real files that in template copying have already been done, so we dont want a warning message.
+			continue;
 		}
 
-	private:
-		const TMap<FNameEntryId, int32>& NameToIndexMap;
-		FString ErrorMessage;
-	};
+		if (FPackageName::DoesPackageExist(PackageName, &SrcFilename))
+		{
+			FString DestFilename = FPackageName::LongPackageNameToFilename(DestPackage, FString(FPathViews::GetExtension(SrcFilename, true)));
+			FilePathRenameMap.Add({ MoveTemp(SrcFilename), MoveTemp(DestFilename) });
+		}
+		else
+		{
+			UE_LOG(LogAssetHeaderPatcher, Warning, TEXT("{%s} package does not exist, and will not be patched."), *PackageName);
+		}
+	}
+}
 
-	enum class EPatchedSection
-	{
-		Summary,
-		NameTable,
-		SoftPathTable,
-		GatherableTextDataTable,
-		ImportTable,
-		ExportTable,
-		SoftPackageReferencesTable,
-		ThumbnailTable,
-		AssetRegistryData
-	};
+void FAssetHeaderPatcher::FContext::GeneratePackagePathsFromFilePaths(const FString& InSrcRoot, const FString& InDstRoot, const FString& InSrcBaseDir)
+{
+	const FString SourceContentPath = FPaths::Combine(InSrcBaseDir, TEXT("Content"));
+	for (const TTuple<FString, FString>& SourceAndDest : FilePathRenameMap)
+	{
+		const FString& SrcFileName = SourceAndDest.Key;
+
+		if (FPaths::IsUnderDirectory(SrcFileName, SourceContentPath))
+		{
+			if (FStringView RelativePkgPath; FPathViews::TryMakeChildPathRelativeTo(SrcFileName, SourceContentPath, RelativePkgPath))
+			{
+				RelativePkgPath = FPathViews::GetBaseFilenameWithPath(RelativePkgPath); // chop the extension
+				if (RelativePkgPath.Len() > 0 && !RelativePkgPath.EndsWith(TEXT("/")))
+				{
+					PackagePathRenameMap.Add(FPaths::Combine(TEXT("/"), InSrcRoot, RelativePkgPath),
+						FPaths::Combine(TEXT("/"), InDstRoot, RelativePkgPath));
+				}
+			}
+		}
+	}
+}
 
-	struct FSectionData
+void FAssetHeaderPatcher::FContext::GatherDependentPackages()
+{
+	// Paths under the __External root drop the package root, so create mappings, per plugin, 
+	// we can leverage when handling those cases where the package path may have been remapped
+	TMap<FString, TMap<FString, FString>> PluginExternalMappings;
+	for (const TPair<FString, FString>& SrcDstPair : PackagePathRenameMap)
 	{
-		EPatchedSection Section = EPatchedSection::Summary;
-		int64 Offset = 0;
-		int64 Size = 0;
-		bool bRequired = false;
-	};
+		const FString& Src = SrcDstPair.Key;
+		const FString& Dst = SrcDstPair.Value;
+
+		FStringView SrcPackageRoot;
+		FStringView SrcPackagePath;
+		FStringView SrcPackageName;
+		SplitLongPackageName(Src, SrcPackageRoot, SrcPackagePath, SrcPackageName);
+
+		FStringView DstPackageRoot;
+		FStringView DstPackagePath;
+		FStringView DstPackageName;
+		SplitLongPackageName(Dst, DstPackageRoot, DstPackagePath, DstPackageName);
+
+		TMap<FString, FString>& ExternalMappings = PluginExternalMappings.FindOrAddByHash(GetTypeHash(SrcPackageRoot), FString(SrcPackageRoot));
+		FStringView SrcPath = SrcPackagePath.IsEmpty() ? SrcPackageName : SrcPackagePath;
+		FStringView DstPath = DstPackagePath.IsEmpty() ? DstPackageName : DstPackagePath;
+		ExternalMappings.Add(FString(SrcPath), FString(DstPath));
+
+		// if there is a path
+		if (!SrcPackagePath.IsEmpty())
+		{
+			// add the local path/asset for the case of maps (which we cannot tell at this point)
+			ExternalMappings.Add(FString(SrcPath.GetData()), FString(DstPath.GetData()));
+		}
+	}
 
-	enum class ESummaryOffset
-	{
-		NameTable,
-		SoftObjectPathList,
-		GatherableTextDataTable,
-		ImportTable,
-		ExportTable,
-		DependsTable,
-		SoftPackageReferenceList,
-		SearchableNamesTable,
-		ThumbnailTable,
-		AssetRegistryData,
-		WorldTileInfoData,
-		PreloadDependency, // Should not be present - only for cooked data
-		BulkData,
-		PayloadToc
-	};
+	TMap<FString, FString> Result;
+	IAssetRegistry& Registry = *IAssetRegistry::Get();
 
-	// To override MemoryReaders FName method
-	class FReadFNameAs2IntFromMemoryReader final : public FLargeMemoryReader
+	TArray< TTuple<FString, FString> > ToProcess;
+	Algo::Copy(PackagePathRenameMap, ToProcess);
+
+	TStringBuilder<NAME_SIZE> SrcDependencyBuilder;
+	while (ToProcess.Num())
 	{
-	public:
-		FReadFNameAs2IntFromMemoryReader(TArray<FName>& InNameTable, const uint8* InData, const int64 Num, ELargeMemoryReaderFlags InFlags = ELargeMemoryReaderFlags::None, const FName InArchiveName = NAME_None)
-			: FLargeMemoryReader(InData, Num, InFlags, InArchiveName)
-			, NameTable(InNameTable)
+		TTuple<FString, FString> Package = ToProcess.Pop();
+
+		if (Result.Contains(Package.Key))
+		{
+			continue;
+		}
+
+		// Become a patching name even if it doesn't have a file.
+		Result.Add({ Package.Key, Package.Value });
+
+		TArray<FName> Dependencies;
+		if (!Registry.GetDependencies(FName(*Package.Key), Dependencies))
 		{
+			continue;
 		}
 
-		// FLargeMemoryReader falls back to FMemoryArchive's imp of this method.
-		// which uses strings as the format for FName.
-		// We need the 2xint32 version when decoding the current file formats. 
-		virtual FArchive& operator<<(FName& OutName) override
+		FStringView SrcPackageRoot = FPackageName::SplitPackageNameRoot(Package.Key, nullptr);
+		FStringView DstPackageRoot = FPackageName::SplitPackageNameRoot(Package.Value, nullptr);
+		for (const FName Dependency : Dependencies)
 		{
-			int32 NameIndex;
-			int32 Number;
-			FArchive& Ar = *this;
-			Ar << NameIndex;
-			Ar << Number;
+			Dependency.ToString(SrcDependencyBuilder);
+			FStringView SrcDependency = SrcDependencyBuilder.ToView();
+
+			if (PackagePathRenameMap.FindByHash(GetTypeHash(SrcDependency), SrcDependency))
+			{
+				// We already handled this mapping
+				continue;
+			}
+
+			FStringView SrcDependencyPackageRoot;
+			FStringView SrcDependencyPackagePath;
+			FStringView SrcDependencyPackageName;
+			SplitLongPackageName(SrcDependency, SrcDependencyPackageRoot, SrcDependencyPackagePath, SrcDependencyPackageName);
+			check(!SrcDependencyPackageRoot.IsEmpty());
+
+			// Only consider dependency paths that are for the same package as our src->dst mapping
+			// If the src mapping doesn't begin with a '/' the package name will be empty, since the path isn't a package path
+			if (SrcDependencyPackageRoot != SrcPackageRoot)
+			{
+				continue;
+			}
+
+			TStringBuilder<NAME_SIZE> DstDependencyString;
 
-			if (NameTable.IsValidIndex(NameIndex))
+			// Special handling for external references. The __External[Actors__|Objects__] directory is always under the package root, may contain an
+			// arbitrary amount of subdirs but then ends with two hash subdirs. The path between the __External[Actors__|Objects__] and the two hash dirs
+			// may need remapping so we look at our external mappings to do so.
+			bool bHasExternalActorDir = SrcDependencyPackagePath.StartsWith(FPackagePath::GetExternalActorsFolderName());
+			bool bHasExternalObjectsDir = !bHasExternalActorDir && SrcDependencyPackagePath.StartsWith(FPackagePath::GetExternalObjectsFolderName());
+			if (bHasExternalActorDir || bHasExternalObjectsDir)
 			{
-				FNameEntryId MappedName = NameTable[NameIndex].GetDisplayIndex();
-				OutName = FName::CreateFromDisplayId(MappedName, Number);
+				int32 RightPartStartPos;
+				if (!SrcDependencyPackagePath.FindChar(TEXT('/'), RightPartStartPos))
+				{
+					// This is a path to only the special directory, skip it no remapping is needed
+					continue;
+				}
+				RightPartStartPos++; // Skip past the '/'
+
+				// Find the start of the two hash dirs
+				// e.g. __ExternalActors__/path/of/interest/A/A9, we only want 'path/of/interest'
+				FStringView ExternalPackagePath(SrcDependencyPackagePath.GetData() + RightPartStartPos, SrcDependencyPackagePath.Len() - RightPartStartPos);
+				int32 HashDirStartPos = 0;
+				int32 NumHashDirsToStrip = 2;
+				while (NumHashDirsToStrip--)
+				{
+					if (ExternalPackagePath.FindLastChar(TEXT('/'), HashDirStartPos))
+					{
+						ExternalPackagePath.LeftChopInline(ExternalPackagePath.Len() - HashDirStartPos);
+					}
+				}
+
+				// Our __External[Actors|Objects]__ path is malformed
+				if (HashDirStartPos == INDEX_NONE)
+				{
+					continue;
+				}
+
+				const int32 HashPathOffset = RightPartStartPos + HashDirStartPos;
+				FStringView HashPath(SrcDependencyPackagePath.GetData() + HashPathOffset, SrcDependencyPackagePath.Len() - HashPathOffset);
+				const TMap<FString, FString>* ExternalMappings = PluginExternalMappings.FindByHash(GetTypeHash(SrcPackageRoot), SrcPackageRoot);
+				if (!ExternalMappings)
+				{
+					// We have no mapping for this dependency's external actors/objects
+					continue;
+				}
+				const FString* DstExternalPackagePath = ExternalMappings->FindByHash(GetTypeHash(ExternalPackagePath), ExternalPackagePath);
+
+				DstDependencyString.AppendChar(TEXT('/'));
+				DstDependencyString.Append(DstPackageRoot);
+				DstDependencyString.AppendChar(TEXT('/'));
+				DstDependencyString.Append(bHasExternalActorDir ? FPackagePath::GetExternalActorsFolderName() : FPackagePath::GetExternalObjectsFolderName());
+				DstDependencyString.AppendChar(TEXT('/'));
+				DstDependencyString.Append(DstExternalPackagePath ? *DstExternalPackagePath : ExternalPackagePath);
+				DstDependencyString.Append(HashPath); // HashPath already contains the leading '/'
+				DstDependencyString.AppendChar(TEXT('/'));
+				DstDependencyString.Append(SrcDependencyPackageName);
 			}
 			else
 			{
-				OutName = FName();
-				SetCriticalError();
+				// We aren't handling a special directory so replace the package root
+				DstDependencyString.AppendChar(TEXT('/'));
+				DstDependencyString.Append(DstPackageRoot);
+				DstDependencyString.AppendChar(TEXT('/'));
+
+				if (!SrcDependencyPackagePath.IsEmpty())
+				{
+					DstDependencyString.Append(SrcDependencyPackagePath);
+					DstDependencyString.AppendChar(TEXT('/'));
+				}
+
+				DstDependencyString.Append(SrcDependencyPackageName);
 			}
 
-			return *this;
+			// If a dep start with the package name, then we are going to copy the asset.
+			// but we need to recurse on this asset as it may have sub dependencies we don't know of yet.
+			ToProcess.Add({ FString(SrcDependency), DstDependencyString.ToString() });
 		}
+	}
 
-		virtual FString GetArchiveName() const override
-		{
-			return TEXT("FReadFNameAs2IntFromMemoryReader");
-		}
-	private:
-		TArray<FName>& NameTable;
-	};
+	PackagePathRenameMap = MoveTemp(Result);
+}
 
-	struct FSummaryOffsetMeta
+void FAssetHeaderPatcher::FContext::GenerateAdditionalRemappings()
+{
+	TArray<FCoreRedirect> ExternalObjectRedirects;
+	TStringBuilder<24> ExternalActorsFolderBuilder;
+	ExternalActorsFolderBuilder << FPackagePath::GetExternalActorsFolderName() << TEXT("/");
+	const FStringView ExternalActorsFolder = ExternalActorsFolderBuilder.ToView();
+
+	TStringBuilder<24> ExternalObjectsFolderBuilder;
+	ExternalObjectsFolderBuilder << FPackagePath::GetExternalObjectsFolderName() << TEXT("/");
+	const FStringView ExternalObjectsFolder = ExternalObjectsFolderBuilder.ToView();
+
+	TStringBuilder<NAME_SIZE> SrcNameBuilder;
+	TStringBuilder<NAME_SIZE> DstNameBuilder;
+	for (const TTuple<FString, FString>& Package : PackagePathRenameMap)
 	{
-		// NOTE: The offsets in Summary get to a max of 312 bytes.
-		// So we could drop this to a uint16 but that is probably overkill at this point.
-		uint32 Offset : 31;
-		uint32 bIs64Bit : 1;
+		const FString& SrcNameString = Package.Key;
+		const FString& DstNameString = Package.Value;
 
-		int64 Value(FPackageFileSummary& Summary) const
+		bool bIsExternalObjectOrActor = false;
+		FStringView SrcPackageName;
 		{
-			intptr_t Ptr = reinterpret_cast<intptr_t>(&Summary) + Offset;
-			if (bIs64Bit)
-			{
-				return *reinterpret_cast<int64*>(Ptr);
-			}
-			else
+			FStringView SrcPackageRoot;
+			FStringView SrcPackagePath;
+			if (!ensure(SplitLongPackageName(SrcNameString, SrcPackageRoot, SrcPackagePath, SrcPackageName))
+				|| SrcPackagePath.StartsWith(ExternalActorsFolder)
+				|| SrcPackagePath.StartsWith(ExternalObjectsFolder))
 			{
-				return *reinterpret_cast<int32*>(Ptr);
+				bIsExternalObjectOrActor = true;
 			}
 		}
 
-		void PatchOffsetValue(FPackageFileSummary& Summary, int64 Value) const
+		// /Path/To/Package mapping
 		{
-			intptr_t Ptr = reinterpret_cast<intptr_t>(&Summary) + Offset;
-			if (bIs64Bit)
+			FCoreRedirect PackageRedirect(ECoreRedirectFlags::Type_Package,
+				FCoreRedirectObjectName(SrcNameString),
+				FCoreRedirectObjectName(DstNameString));
+
+			if (bIsExternalObjectOrActor)
 			{
-				int64& Dst = *reinterpret_cast<int64*>(Ptr);
-				Dst += Value;
+				// The other mappings below don't apply to ExternalActors or ExternalObjects so we skip them
+				// now that we have a PackagePath mapping for them
+				ExternalObjectRedirects.Emplace(MoveTemp(PackageRedirect));
+				continue;
 			}
 			else
 			{
-				int32& Dst = *reinterpret_cast<int32*>(Ptr);
-				*reinterpret_cast<int32*>(Ptr) = IntCastChecked<int32>((int64)Dst + Value);
+				Redirects.Emplace(MoveTemp(PackageRedirect));
 			}
 		}
-	};
 
-	const FSummaryOffsetMeta OffsetTable[] = {
+		FStringView DstPackageName = FPathViews::GetBaseFilename(DstNameString);
+
+		// Path.ObjectName mapping
+		{
+			SrcNameBuilder.Reset();
+			SrcNameBuilder.Append(SrcNameString);
+			SrcNameBuilder.AppendChar(TEXT('.'));
+			SrcNameBuilder.Append(SrcPackageName);
+
+			DstNameBuilder.Reset();
+			DstNameBuilder.Append(DstNameString);
+			DstNameBuilder.AppendChar(TEXT('.'));
+			DstNameBuilder.Append(DstPackageName);
+
+			FCoreRedirect PackageObjectRedirect(ECoreRedirectFlags::Type_Package | ECoreRedirectFlags::Type_Object,
+				FCoreRedirectObjectName(SrcNameBuilder.ToString()),
+				FCoreRedirectObjectName(DstNameBuilder.ToString()));
+			Redirects.Emplace(MoveTemp(PackageObjectRedirect));
+		}
+
+		// MaterialFunctionInterface "EditorOnlyData"
+		{
+			SrcNameBuilder.Reset();
+			SrcNameBuilder.Append(SrcNameString);
+			SrcNameBuilder.AppendChar(TEXT('.'));
+			SrcNameBuilder.Append(SrcPackageName);
+			SrcNameBuilder.Append(TEXT("EditorOnlyData"));
+
+			DstNameBuilder.Reset();
+			DstNameBuilder.Append(DstNameString);
+			DstNameBuilder.AppendChar(TEXT('.'));
+			DstNameBuilder.Append(DstPackageName);
+			DstNameBuilder.Append(TEXT("EditorOnlyData"));
+
+			FCoreRedirect BlueprintClassRedirect(ECoreRedirectFlags::Type_Class | ECoreRedirectFlags::Type_Package,
+				FCoreRedirectObjectName(SrcNameBuilder.ToString()),
+				FCoreRedirectObjectName(DstNameBuilder.ToString()));
+			Redirects.Emplace(MoveTemp(BlueprintClassRedirect));
+		}
+
+		// Compiled Blueprint class names
+		{
+			SrcNameBuilder.Reset();
+			SrcNameBuilder.Append(SrcNameString);
+			SrcNameBuilder.AppendChar(TEXT('.'));
+			SrcNameBuilder.Append(SrcPackageName);
+			SrcNameBuilder.Append(TEXT("_C"));
+
+			DstNameBuilder.Reset();
+			DstNameBuilder.Append(DstNameString);
+			DstNameBuilder.AppendChar(TEXT('.'));
+			DstNameBuilder.Append(DstPackageName);
+			DstNameBuilder.Append(TEXT("_C"));
+
+			FCoreRedirect BlueprintClassRedirect(ECoreRedirectFlags::Type_Class | ECoreRedirectFlags::Type_Package,
+				FCoreRedirectObjectName(SrcNameBuilder.ToString()),
+				FCoreRedirectObjectName(DstNameBuilder.ToString()));
+			Redirects.Emplace(MoveTemp(BlueprintClassRedirect));
+		}
+
+		// Blueprint generated class default object
+		{
+			SrcNameBuilder.Reset();
+			SrcNameBuilder.Append(SrcNameString);
+			SrcNameBuilder.AppendChar(TEXT('.'));
+			SrcNameBuilder.Append(DEFAULT_OBJECT_PREFIX);
+			SrcNameBuilder.Append(SrcPackageName);
+			SrcNameBuilder.Append(TEXT("_C"));
+
+			DstNameBuilder.Reset();
+			DstNameBuilder.Append(DstNameString);
+			DstNameBuilder.AppendChar(TEXT('.'));
+			DstNameBuilder.Append(DEFAULT_OBJECT_PREFIX);
+			DstNameBuilder.Append(DstPackageName);
+			DstNameBuilder.Append(TEXT("_C"));
+
+			FCoreRedirect DefaultBlueprintClassRedirect(ECoreRedirectFlags::Type_Class | ECoreRedirectFlags::Type_Package,
+				FCoreRedirectObjectName(SrcNameBuilder.ToString()),
+				FCoreRedirectObjectName(DstNameBuilder.ToString()));
+			Redirects.Emplace(MoveTemp(DefaultBlueprintClassRedirect));
+		}
+	}
+
+	// For best-effort string matches. Intentionally excluding external objects as AssetRegistry Tag data
+	// can't refer to these paths in a manner that we can't deduce from the redirects themselves
+	for (auto& Redirect : Redirects)
+	{
+		const FCoreRedirectObjectName& SrcName = Redirect.OldName;
+		const FCoreRedirectObjectName& DstName = Redirect.NewName;
+
+		StringReplacements.Add(SrcName.ObjectName.ToString(), DstName.ObjectName.ToString());
+		StringReplacements.Add(SrcName.PackageName.ToString(), DstName.PackageName.ToString());
+		StringReplacements.Add(SrcName.ToString(), DstName.ToString());
+
+		// Tag data can contain VersePaths which are like Top-Level Asset Paths
+		// but with a mountpoint prefix and only '/' delimiters
+		for (FString& VerseMount : VerseMountPoints)
+		{
+			SrcNameBuilder.Reset();
+			SrcNameBuilder.AppendChar(TEXT('/'));
+			SrcNameBuilder.Append(VerseMount);
+			SrcName.PackageName.AppendString(SrcNameBuilder);
+			SrcNameBuilder.AppendChar(TEXT('/'));
+			SrcName.ObjectName.AppendString(SrcNameBuilder);
+
+			DstNameBuilder.Reset();
+			DstNameBuilder.AppendChar(TEXT('/'));
+			DstNameBuilder.Append(VerseMount);
+			DstName.PackageName.AppendString(DstNameBuilder);
+			DstNameBuilder.AppendChar(TEXT('/'));
+			DstName.ObjectName.AppendString(DstNameBuilder);
+			StringReplacements.Add(SrcNameBuilder.ToString(), DstNameBuilder.ToString());
+		}
+	}
+
+	// Now that we have generated the string matches above, add the external redirects
+	Redirects.Append(ExternalObjectRedirects);
+
+	// Add prefix redirects for any mountpoint replacements
+	TMap<FString, FString> FormattedStringMountReplacements;
+	FormattedStringMountReplacements.Reserve(StringMountReplacements.Num());
+	for (const auto& MountPointPair : StringMountReplacements)
+	{
+		const FString& SrcMountPoint = MountPointPair.Key;
+		const FString& DstMountPoint = MountPointPair.Value;
+
+		SrcNameBuilder.Reset();
+		SrcNameBuilder.AppendChar(TEXT('/'));
+		SrcNameBuilder.Append(SrcMountPoint);
+		SrcNameBuilder.AppendChar(TEXT('/'));
+
+		DstNameBuilder.Reset();
+		DstNameBuilder.AppendChar(TEXT('/'));
+		DstNameBuilder.Append(DstMountPoint);
+		DstNameBuilder.AppendChar(TEXT('/'));
+
+		FCoreRedirect MountRedirect(ECoreRedirectFlags::Type_Package | ECoreRedirectFlags::Option_MatchPrefix,
+			FCoreRedirectObjectName(SrcNameBuilder.ToString()),
+			FCoreRedirectObjectName(DstNameBuilder.ToString()));
+		Redirects.Emplace(MoveTemp(MountRedirect));
+
+		// Store off the actual mount path prefix to make patching easier later
+		FormattedStringMountReplacements.Add(SrcNameBuilder.ToString(), DstNameBuilder.ToString());
+	}
+	StringMountReplacements = MoveTemp(FormattedStringMountReplacements);
+}
+
+// To override writing of FName's to ensure they have been patched
+class FNamePatchingWriter final : public FArchiveProxy
+{
+public:
+	FNamePatchingWriter(FArchive& InAr, const TMap<FNameEntryId, int32>& InNameToIndexMap)
+		: FArchiveProxy(InAr)
+		, NameToIndexMap(InNameToIndexMap)
+	{
+	}
+
+	virtual ~FNamePatchingWriter() { }
+
+	virtual FArchive& operator<<(FName& Name) override
+	{
+		FNameEntryId EntryId = Name.GetDisplayIndex();
+		const int32* MaybeIndex = NameToIndexMap.Find(EntryId);
+
+		if (MaybeIndex == nullptr)
+		{
+			ErrorMessage += FString::Printf(TEXT("Cannot serialize FName %s because it is not in the name table for %s\n"), *Name.ToString(), *GetArchiveName());
+			SetCriticalError();
+			return *this;
+		}
+
+		int32 Index = *MaybeIndex;
+		int32 Number = Name.GetNumber();
+
+		FArchive& Ar = *this;
+		Ar << Index;
+		Ar << Number;
+
+		return *this;
+	}
+
+	const FString& GetErrorMessage() const
+	{
+		return ErrorMessage;
+	}
+
+private:
+	const TMap<FNameEntryId, int32>& NameToIndexMap;
+	FString ErrorMessage;
+};
+
+enum class EPatchedSection
+{
+	Summary,
+	NameTable,
+	SoftPathTable,
+	GatherableTextDataTable,
+	SearchableNamesMap,
+	ImportTable,
+	ExportTable,
+	SoftPackageReferencesTable,
+	ThumbnailTable,
+	AssetRegistryData
+};
+
+struct FSectionData
+{
+	EPatchedSection Section = EPatchedSection::Summary;
+	int64 Offset = 0;
+	int64 Size = 0;
+	bool bRequired = false;
+};
+
+enum class ESummaryOffset
+{
+	NameTable,
+	SoftObjectPathList,
+	GatherableTextDataTable,
+	ImportTable,
+	ExportTable,
+	DependsTable,
+	SoftPackageReferenceList,
+	SearchableNamesMap,
+	ThumbnailTable,
+	AssetRegistryData,
+	WorldTileInfoData,
+	PreloadDependency, // Should not be present - only for cooked data
+	BulkData,
+	PayloadToc
+};
+
+// To override MemoryReaders FName method
+class FReadFNameAs2IntFromMemoryReader final : public FLargeMemoryReader
+{
+public:
+	FReadFNameAs2IntFromMemoryReader(TArray<FName>& InNameTable, const uint8* InData, const int64 Num, ELargeMemoryReaderFlags InFlags = ELargeMemoryReaderFlags::None, const FName InArchiveName = NAME_None)
+		: FLargeMemoryReader(InData, Num, InFlags, InArchiveName)
+		, NameTable(InNameTable)
+	{
+	}
+
+	// FLargeMemoryReader falls back to FMemoryArchive's imp of this method.
+	// which uses strings as the format for FName.
+	// We need the 2xint32 version when decoding the current file formats. 
+	virtual FArchive& operator<<(FName& OutName) override
+	{
+		int32 NameIndex;
+		int32 Number;
+		FArchive& Ar = *this;
+		Ar << NameIndex;
+		Ar << Number;
+
+		if (NameTable.IsValidIndex(NameIndex))
+		{
+			FNameEntryId MappedName = NameTable[NameIndex].GetDisplayIndex();
+			OutName = FName::CreateFromDisplayId(MappedName, Number);
+		}
+		else
+		{
+			OutName = FName();
+			SetCriticalError();
+		}
 
-	#define UE_POPULATE_OFFSET_INFO(NAME)					\
+		return *this;
+	}
+
+	virtual FString GetArchiveName() const override
+	{
+		return TEXT("FReadFNameAs2IntFromMemoryReader");
+	}
+private:
+	TArray<FName>& NameTable;
+};
+
+struct FSummaryOffsetMeta
+{
+	// NOTE: The offsets in Summary get to a max of 312 bytes.
+	// So we could drop this to a uint16 but that is probably overkill at this point.
+	uint32 Offset : 31;
+	uint32 bIs64Bit : 1;
+
+	int64 Value(FPackageFileSummary& Summary) const
+	{
+		intptr_t Ptr = reinterpret_cast<intptr_t>(&Summary) + Offset;
+		if (bIs64Bit)
+		{
+			return *reinterpret_cast<int64*>(Ptr);
+		}
+		else
+		{
+			return *reinterpret_cast<int32*>(Ptr);
+		}
+	}
+
+	void PatchOffsetValue(FPackageFileSummary& Summary, int64 Value) const
+	{
+		intptr_t Ptr = reinterpret_cast<intptr_t>(&Summary) + Offset;
+		if (bIs64Bit)
+		{
+			int64& Dst = *reinterpret_cast<int64*>(Ptr);
+			Dst += Value;
+		}
+		else
+		{
+			int32& Dst = *reinterpret_cast<int32*>(Ptr);
+			*reinterpret_cast<int32*>(Ptr) = IntCastChecked<int32>((int64)Dst + Value);
+		}
+	}
+};
+
+const FSummaryOffsetMeta OffsetTable[] = {
+
+#define UE_POPULATE_OFFSET_INFO(NAME)					\
 		(uint32)STRUCT_OFFSET(FPackageFileSummary, NAME),	\
 		std::is_same_v<decltype(((FPackageFileSummary*)0)->NAME), int64>
 
@@ -210,50 +739,37 @@ namespace
 			{ UE_POPULATE_OFFSET_INFO(PayloadTocOffset) },
 
 #undef UE_POPULATE_OFFSET_INFO
-	};
+};
 
-	void PatchSummaryOffsets(FPackageFileSummary& Dst, int64 OffsetFrom, int64 OffsetDelta)
+void PatchSummaryOffsets(FPackageFileSummary& Dst, int64 OffsetFrom, int64 OffsetDelta)
+{
+	if (!OffsetDelta)
 	{
-		if (!OffsetDelta)
-		{
-			return;
-		}
-
-		for (const FSummaryOffsetMeta& OffsetData : OffsetTable)
-		{
-			if (OffsetData.Value(Dst) > OffsetFrom)
-			{
-				OffsetData.PatchOffsetValue(Dst, OffsetDelta);
-			}
-		}
-	};
+		return;
+	}
 
-	FAssetDataTagMap MakeTagMap(const TArray<UE::AssetRegistry::FDeserializeTagData>& TagData)
+	for (const FSummaryOffsetMeta& OffsetData : OffsetTable)
 	{
-		FAssetDataTagMap Out;
-		Out.Reserve(TagData.Num());
-		for (const UE::AssetRegistry::FDeserializeTagData& Tag : TagData)
+		if (OffsetData.Value(Dst) > OffsetFrom)
 		{
-			if (!Tag.Key.IsEmpty() && !Tag.Value.IsEmpty())
-			{
-				Out.Add(*Tag.Key, Tag.Value);
-			}
+			OffsetData.PatchOffsetValue(Dst, OffsetDelta);
 		}
-
-		return Out;
 	}
+};
 
-	FStringView Find(const TMap<FString, FString>& Table, FStringView Needle)
+FAssetDataTagMap MakeTagMap(const TArray<UE::AssetRegistry::FDeserializeTagData>& TagData)
+{
+	FAssetDataTagMap Out;
+	Out.Reserve(TagData.Num());
+	for (const UE::AssetRegistry::FDeserializeTagData& Tag : TagData)
 	{
-		uint32 NeedleHash = TMap<FString, FString>::KeyFuncsType::GetKeyHash<FStringView>(Needle);
-		const FString* MaybeNewItem = Table.FindByHash<FStringView>(NeedleHash, Needle);
-		if (MaybeNewItem)
+		if (!Tag.Key.IsEmpty() && !Tag.Value.IsEmpty())
 		{
-			return *MaybeNewItem;
+			Out.Add(*Tag.Key, Tag.Value);
 		}
-
-		return {};
 	}
+
+	return Out;
 }
 
 // The information we need in the task to do patching.
@@ -262,30 +778,52 @@ class FAssetHeaderPatcherInner
 public:
 	using EResult = FAssetHeaderPatcher::EResult;
 
-	FAssetHeaderPatcherInner(const FString& InSrcAsset, const FString& InDstAsset, const TMap<FString, FString>& InSearchAndReplace, FArchive* InDstArchive = nullptr, bool bInBespokeSearchInUse = false)
+	struct FThumbnailEntry
+	{
+		FString ObjectShortClassName;
+		FString ObjectPathWithoutPackageName;
+		int32 FileOffset = 0;
+		int32 Delta = 0;
+	};
+
+	FAssetHeaderPatcherInner(const FString& InSrcAsset, const FString& InDstAsset, const TMap<FString, FString>& InStringReplacements, const TMap<FString, FString>& InStringMountPointReplacements, FArchive* InDstArchive = nullptr)
 		: SrcAsset(InSrcAsset)
 		, DstAsset(InDstAsset)
-		, SearchAndReplace(InSearchAndReplace)
-		, bBespokeSearchInUse(bInBespokeSearchInUse)
+		, StringReplacements(InStringReplacements)
+		, StringMountPointReplacements(InStringMountPointReplacements)
 		, DstArchive(InDstArchive)
+		, bPatchPrimaryAssetTag(false)
+		, bIsNonOneFilePerActorPackage(false)
 	{
+		for (auto TagToIgnore : TagsToIgnore)
+		{
+			IgnoredTags.Add(TagToIgnore);
+		}
 	}
 
 	bool DoPatch(FString& InOutString);
 	bool DoPatch(FName& InOutName);
+	bool DoPatch(FObjectResource& InOutResource, bool bIsExport, FName& OutPatchedObjectName);
 	bool DoPatch(FSoftObjectPath& InOutSoft);
 	bool DoPatch(FTopLevelAssetPath& InOutPath);
-	bool DoPatch(FGatherableTextData& InGatherablerTextData);
+	bool DoPatch(FGatherableTextData& InOutGatherablerTextData);
+	bool DoPatch(FThumbnailEntry& InOutThumbnailEntry);
+	bool RemapFName(FName SrcName, FName DstName);
+	bool ShouldReplaceMountPoint(const FStringView InPath, FStringView& OutSrcMountPoint, FStringView& OutDstMountPoint);
+	void PatchNameTable();
 
 	FAssetHeaderPatcher::EResult PatchHeader();
 	FAssetHeaderPatcher::EResult PatchHeader_Deserialize();
 	void PatchHeader_PatchSections();
 	FAssetHeaderPatcher::EResult PatchHeader_WriteDestinationFile();
+	void DumpState(FStringView InDir);
+
+	TSet<FString> IgnoredTags;
 
 	const FString& SrcAsset;
 	const FString& DstAsset;
-	const TMap<FString, FString>& SearchAndReplace;
-	bool  bBespokeSearchInUse;
+	const TMap<FString, FString>& StringReplacements;
+	const TMap<FString, FString>& StringMountPointReplacements;
 	FArchive* DstArchive = nullptr;
 	TUniquePtr<FArchive> DstArchiveOwner;
 
@@ -301,26 +839,31 @@ public:
 		int64 ExportTableSize = -1;
 		int64 SoftPackageReferencesListSize = -1;
 		int64 ThumbnailTableSize = -1;
+		int64 SearchableNamesMapSize = -1;
 		int64 AssetRegistryDataSize = -1;
 		int64 PackageTrailerSize = -1;
 	};
 
-	struct FThumbnailEntry
-	{
-		FString ObjectShortClassName;
-		FString ObjectPathWithoutPackageName;
-		int32 FileOffset = 0;
-	};
-
 	FHeaderInformation HeaderInformation;
 	FPackageFileSummary Summary;
+	FName OriginalPackagePath;					 // e.g. "/MountName/TopLevelPackageName"
+	FName OriginalNonOneFilePerActorPackagePath; // e.g. "/MountName/MountName"
+	FString OriginalPrimaryAssetName;			 // e.g. "MountName"
+	bool bPatchPrimaryAssetTag;
+	bool bIsNonOneFilePerActorPackage;
+
+	// NameTable Members
 	TArray<FName> NameTable;
 	TMap<FNameEntryId, int32> NameToIndexMap;
+	TMap<FNameEntryId, FNameEntryId> RenameMap;
+	TSet<FNameEntryId> AddedNames;
+
 	TArray<FSoftObjectPath> SoftObjectPathTable;
 	TArray<FGatherableTextData> GatherableTextDataTable;
 	TArray<FObjectImport> ImportTable;
 	TArray<FObjectExport> ExportTable;
 	TArray<FName> SoftPackageReferencesTable;
+	TMap<FPackageIndex, TArray<FName>> SearchableNamesMap;
 	TArray<FThumbnailEntry> ThumbnailTable;
 
 	// Asset registry data information
@@ -339,9 +882,9 @@ public:
 	FAssetRegistryData AssetRegistryData;
 };
 
-FAssetHeaderPatcher::EResult FAssetHeaderPatcher::DoPatch(const FString& InSrcAsset, const FString& InDstAsset, const TMap<FString, FString>& InSearchAndReplace, bool bInStarRestrictionInUse)
+FAssetHeaderPatcher::EResult FAssetHeaderPatcher::DoPatch(const FString& InSrcAsset, const FString& InDstAsset, const FContext& InContext)
 {
-	FAssetHeaderPatcherInner Inner(InSrcAsset, InDstAsset, InSearchAndReplace, nullptr, bInStarRestrictionInUse);
+	FAssetHeaderPatcherInner Inner(InSrcAsset, InDstAsset, InContext.StringReplacements, InContext.StringMountReplacements);
 
 	if (!FFileHelper::LoadFileToArray(Inner.SrcBuffer, *Inner.SrcAsset))
 	{
@@ -354,25 +897,118 @@ FAssetHeaderPatcher::EResult FAssetHeaderPatcher::DoPatch(const FString& InSrcAs
 	}
 }
 
+void FAssetHeaderPatcher::Reset()
+{
+	ErroredFiles.Empty();
+	PatchedFiles.Empty();
 
-FAssetHeaderPatcher::EResult FAssetHeaderPatcher::Test_DoPatch(FArchive& InReader, FArchive& InWriter,
-	const TMap<FString, FString>& InSearchAndReplace)
+	PatchingTask = UE::Tasks::FTask();
+	Status = EResult::NotStarted;
+	bCancelled = false;
+}
+void FAssetHeaderPatcher::SetContext(FContext InContext)
 {
-	FAssetHeaderPatcherInner Inner(InReader.GetArchiveName(), InWriter.GetArchiveName(), InSearchAndReplace, &InWriter);
+	checkf(!IsPatching(), TEXT("Cannot set the patcher context while patching"));
+	Context = InContext;
+	Reset();
+}
 
-	InReader.Seek(0);
-	Inner.SrcBuffer.SetNumUninitialized(InReader.TotalSize());
-	InReader.Serialize(Inner.SrcBuffer.GetData(), Inner.SrcBuffer.Num());
+UE::Tasks::FTask FAssetHeaderPatcher::PatchAsync(int32* InOutNumFilesToPatch, int32* InOutNumFilesPatched)
+{
+	return PatchAsync(InOutNumFilesToPatch, InOutNumFilesPatched, FAssetHeaderPatcherCompletionDelegate(), FAssetHeaderPatcherCompletionDelegate());
+}
 
-	if (InReader.IsError())
+UE::Tasks::FTask FAssetHeaderPatcher::PatchAsync(int32* InOutNumFilesToPatch, int32* InOutNumFilesPatched, FAssetHeaderPatcherCompletionDelegate InOnSuccess, FAssetHeaderPatcherCompletionDelegate InOnError)
+{
+	PatchedFiles = Context.FilePathRenameMap;
+	if (InOutNumFilesToPatch)
 	{
-		UE_LOG(LogAssetHeaderPatcher, Error, TEXT("Failed to load %s"), *Inner.SrcAsset);
-		return FAssetHeaderPatcherInner::EResult::ErrorFailedToLoadSourceAsset;
+		*InOutNumFilesToPatch = PatchedFiles.Num();
 	}
-	else
+
+	// Before we start patching we need to apply any patching redirects that exist
+	FCoreRedirects::AddRedirectList(Context.Redirects, TEXT("Asset Header Patcher"));
+
+	// Spawn tasks (Scatter)
+	UE::Tasks::FTask PatchAssetsCleanupTask;
+	TArray<UE::Tasks::FTask> PatchAssetTasks;
+
+	// Note we are scheduling and launching tasks one at a time rather than preparing all jobs and launching all at once.
+	// While this means more overhead scheduling, it means that we won't have many tasks all hit the filesystem at the same time
+	// attempting to read and (more importantly) write to disk at the exact same time.
+	constexpr bool bSingleThreaded = true; // Useful for debugging
+	for (const TTuple<FString, FString>& Filename : PatchedFiles)
 	{
-		return Inner.PatchHeader();
+		auto DoPatchFn = [this, 
+			SrcFilename = Filename.Key,
+			DstFilename = Filename.Value,
+			NumPatched = InOutNumFilesPatched, 
+			OnSuccess = InOnSuccess,
+			OnError = InOnError]()
+			{
+				// Even if we are cancelled, increment our progress
+				if (NumPatched)
+				{
+					std::atomic_ref(*NumPatched).fetch_add(1, std::memory_order_relaxed);
+				}
+
+				if (bCancelled)
+				{
+					return;
+				}
+
+				FAssetHeaderPatcher::EResult Result = FAssetHeaderPatcher::DoPatch(SrcFilename, DstFilename, Context);
+				if (Result != FAssetHeaderPatcher::EResult::Success)
+				{
+					FScopeLock Lock(&ErroredFilesLock);
+					// Don't lose our cancelled state, even when there are errors
+					if (Status != EResult::Cancelled)
+					{
+						Status = Result;
+					}
+					ErroredFiles.Add(SrcFilename, Result);
+
+					OnError.ExecuteIfBound(SrcFilename, DstFilename);
+				}
+				else
+				{
+					OnSuccess.ExecuteIfBound(SrcFilename, DstFilename);
+				}
+			};
+
+		if constexpr (bSingleThreaded)
+		{
+			DoPatchFn();
+		}
+		else
+		{
+			PatchAssetTasks.Add(UE::Tasks::Launch(UE_SOURCE_LOCATION, MoveTemp(DoPatchFn)));
+		}
 	}
+
+	// Once all tasks have completed, remove the redirects before we declare Patching complete
+	UE::Tasks::FTask PatcherCleanupTask = UE::Tasks::Launch(UE_SOURCE_LOCATION, [this]()
+		{
+			FCoreRedirects::RemoveRedirectList(Context.Redirects, TEXT("Asset Header Patcher"));
+
+			if (Status != EResult::Cancelled && ErroredFiles.IsEmpty())
+			{
+				Status = EResult::Success;
+			}
+
+			{
+				FScopeLock Lock(&ErroredFilesLock);
+				for (auto& ErroredFile : ErroredFiles)
+				{
+					PatchedFiles.Remove(ErroredFile.Key);
+				}
+			}
+
+		}, UE::Tasks::Prerequisites(PatchAssetTasks));
+
+	Status = EResult::InProgress;
+
+	return PatcherCleanupTask;
 }
 
 FAssetHeaderPatcher::EResult FAssetHeaderPatcherInner::PatchHeader()
@@ -382,8 +1018,26 @@ FAssetHeaderPatcher::EResult FAssetHeaderPatcherInner::PatchHeader()
 	{
 		return Result;
 	}
+	
+	if (DumpOutputDirectory.IsEmpty())
+	{
+		PatchHeader_PatchSections();
+	}
+	else
+	{
+		FString BaseDir = DumpOutputDirectory;
+		FPaths::NormalizeDirectoryName(BaseDir);
+
+		FString BeforeDir = BaseDir / FString(TEXT("Before"));
+		FPaths::RemoveDuplicateSlashes(BeforeDir);
+		DumpState(BeforeDir);
 
-	PatchHeader_PatchSections();
+		PatchHeader_PatchSections();
+
+		FString AfterDir = BaseDir / FString(TEXT("After"));
+		FPaths::RemoveDuplicateSlashes(AfterDir);
+		DumpState(AfterDir);
+	}
 
 	return PatchHeader_WriteDestinationFile();
 }
@@ -395,6 +1049,99 @@ FAssetHeaderPatcher::EResult FAssetHeaderPatcherInner::PatchHeader_Deserialize()
 	MemAr << Summary;
 	HeaderInformation.SummarySize = MemAr.Tell();
 
+	// Summary.PackageName isn't always serialized. In such cases, determine the package name from the file name
+	if (Summary.PackageName.IsEmpty() || Summary.PackageName.Equals(TEXT("None")))
+	{
+		// e.g. "../../Some/Long/Path/MyPlugin/Plugins/MyPackage/Content/TopLevelAssetName.uasset"
+		TStringView Path(SrcAsset);
+
+		static const TStringView ContentDir(TEXT("/Content/"));
+		int32 Pos = Path.Find(ContentDir, ESearchCase::IgnoreCase);
+		if (Pos <= 0)
+		{
+			UE_LOG(LogAssetHeaderPatcher, Error, TEXT("Cannot patch '%s': Package header is missing a 'PackageName' string, nor could a PackageName be deduced."), *SrcAsset);
+			return FAssetHeaderPatcher::EResult::ErrorEmptyRequireSection;
+		}
+		
+		int32 MountNamePos;
+		TStringView LeftPath(Path.GetData(), Pos);
+		if (!LeftPath.FindLastChar(TEXT('/'), MountNamePos))
+		{
+			UE_LOG(LogAssetHeaderPatcher, Error, TEXT("Cannot patch '%s': Package header is missing a 'PackageName' string, nor could a PackageName be deduced."), *SrcAsset);
+			return FAssetHeaderPatcher::EResult::ErrorEmptyRequireSection;
+		}
+
+		int32 ExtensionPos;
+		TStringView RightPath(Path.GetData() + Pos + ContentDir.Len());
+		if (!RightPath.FindLastChar(TEXT('.'), ExtensionPos))
+		{
+			UE_LOG(LogAssetHeaderPatcher, Error, TEXT("Cannot patch '%s': Package header is missing a 'PackageName' string, nor could a PackageName be deduced."), *SrcAsset);
+			return FAssetHeaderPatcher::EResult::ErrorEmptyRequireSection;
+		}
+
+		TStringView MountName(LeftPath.GetData() + MountNamePos, (Pos - MountNamePos) + 1); // + 1 so we can include the '/' from "/Content"
+		TStringView AssetPath(RightPath.GetData(), ExtensionPos);
+		Summary.PackageName.Empty(MountName.Len() + AssetPath.Len());
+		Summary.PackageName.Append(MountName);
+		Summary.PackageName.Append(AssetPath);
+	}
+
+	// Store the original name as an FName as it will be used when
+	// patching paths for other objects in the package
+	{
+		OriginalPackagePath = FName(Summary.PackageName, NAME_NO_NUMBER_INTERNAL);
+
+		// Some ObjectPaths have an implied package, however when it comes to 
+		// non-One File Per Actor packages, the implied package is the map package
+		// so we determine which package we are and cache the map name in case we need it
+		{
+			bIsNonOneFilePerActorPackage = false;
+			TStringBuilder<256> PathBuilder;
+			PathBuilder.AppendChar(TEXT('/'));
+			PathBuilder.Append(FPackagePath::GetExternalActorsFolderName());
+			PathBuilder.AppendChar(TEXT('/'));
+			if (Summary.PackageName.Contains(PathBuilder))
+			{
+				bIsNonOneFilePerActorPackage = true;
+			}
+			else
+			{
+				PathBuilder.Reset();
+				PathBuilder.AppendChar(TEXT('/'));
+				PathBuilder.Append(FPackagePath::GetExternalObjectsFolderName());
+				PathBuilder.AppendChar(TEXT('/'));
+				bIsNonOneFilePerActorPackage = Summary.PackageName.Contains(PathBuilder);
+			}
+
+			int32 SlashPos = INDEX_NONE;
+			FStringView PackageRoot(Summary.PackageName);
+			if (!PackageRoot.FindChar(TEXT('/'), SlashPos) || SlashPos != 0)
+			{
+				UE_LOG(LogAssetHeaderPatcher, Error, TEXT("Cannot patch '%s': PackageName is malformed."), *SrcAsset);
+				return FAssetHeaderPatcher::EResult::ErrorFailedToDeserializeSourceAsset;
+			}
+
+			PackageRoot.RightChopInline(1); // Drop the first slash
+			if (!PackageRoot.FindChar(TEXT('/'), SlashPos))
+			{
+				UE_LOG(LogAssetHeaderPatcher, Error, TEXT("Cannot patch '%s': PackageName is malformed."), *SrcAsset);
+				return FAssetHeaderPatcher::EResult::ErrorFailedToDeserializeSourceAsset;
+			}
+
+			PathBuilder.Reset();
+			PathBuilder.AppendChar(TEXT('/'));
+			PathBuilder.Append(PackageRoot.GetData(), SlashPos);
+			PathBuilder.AppendChar(TEXT('/'));
+			PathBuilder.Append(PackageRoot.GetData(), SlashPos);
+			OriginalNonOneFilePerActorPackagePath = FName(PathBuilder);
+
+			// While here set the OriginalPrimaryAssetName which is used in AssetRegistry Tag lookups for GameFeatureData
+			bPatchPrimaryAssetTag = FPathViews::GetBaseFilename(Summary.PackageName) == TEXT("GameFeatureData");
+			OriginalPrimaryAssetName.Empty();
+			OriginalPrimaryAssetName.Append(PackageRoot.GetData(), SlashPos);
+		}
+	}
+
 	// set version numbers so components branch correctly
 	MemAr.SetUEVer(Summary.GetFileVersionUE());
 	MemAr.SetLicenseeUEVer(Summary.GetFileVersionLicenseeUE());
@@ -424,6 +1171,18 @@ FAssetHeaderPatcher::EResult FAssetHeaderPatcherInner::PatchHeader_Deserialize()
 		}
 
 		HeaderInformation.NameTableSize = MemAr.Tell() - HeaderInformation.SummarySize;
+
+		// Initialize a mapping for Name to index in NameTable as we will use
+		// this for patching in new names and to determine if multiple FNames share the same 
+		// value but might not after patching (i.e. their use of the name differs based on context, and
+		// post-patching the FNames in those contexts no longer match.
+		NameToIndexMap.Empty(NameTable.Num());
+		RenameMap.Reserve(NameTable.Num());
+		AddedNames.Empty();
+		for (int32 i = 0; i < NameTable.Num(); ++i)
+		{
+			NameToIndexMap.Add(NameTable[i].GetDisplayIndex(), i);
+		}
 	}
 
 	if (Summary.SoftObjectPathsCount > 0)
@@ -437,10 +1196,15 @@ FAssetHeaderPatcher::EResult FAssetHeaderPatcherInner::PatchHeader_Deserialize()
 		}
 		HeaderInformation.SoftObjectPathListSize = MemAr.Tell() - Summary.SoftObjectPathsOffset;
 	}
-	else
+	else if(Summary.GetFileVersionUE() >= EUnrealEngineObjectUE5Version::ADD_SOFTOBJECTPATH_LIST)
 	{
 		HeaderInformation.SoftObjectPathListSize = 0;
 	}
+	else
+	{
+		UE_LOG(LogAssetHeaderPatcher, Error, TEXT("Asset '%s' is too old to be used with AssetHeaderPatching. Please resave the file before trying to patch again."), *SrcAsset);
+		return EResult::ErrorUnkownSection;
+	}
 
 	if (Summary.GatherableTextDataCount > 0)
 	{
@@ -529,6 +1293,16 @@ FAssetHeaderPatcher::EResult FAssetHeaderPatcherInner::PatchHeader_Deserialize()
 		HeaderInformation.SoftPackageReferencesListSize = 0;
 	}
 
+	if (Summary.SearchableNamesOffset)
+	{
+		MemAr.Seek(Summary.SearchableNamesOffset);
+		FLinkerTables LinkerTables;
+		LinkerTables.SerializeSearchableNamesMap(MemAr);
+		SearchableNamesMap = MoveTemp(LinkerTables.SearchableNamesMap);
+
+		HeaderInformation.SearchableNamesMapSize = MemAr.Tell() - Summary.SearchableNamesOffset;
+	}
+
 	if (Summary.ThumbnailTableOffset)
 	{
 		MemAr.Seek(Summary.ThumbnailTableOffset);
@@ -587,11 +1361,32 @@ FAssetHeaderPatcher::EResult FAssetHeaderPatcherInner::PatchHeader_Deserialize()
 	return EResult::Success;
 }
 
+bool FAssetHeaderPatcherInner::ShouldReplaceMountPoint(const FStringView InPath, FStringView& OutSrcMountPoint, FStringView& OutDstMountPoint)
+{
+	for (auto& MountPair : StringMountPointReplacements)
+	{
+		const FStringView SrcMount(MountPair.Key);
+		const FStringView DstMount(MountPair.Value);
+
+		if (InPath.StartsWith(SrcMount))
+		{
+			OutSrcMountPoint = SrcMount;
+			OutDstMountPoint = DstMount;
+			return true;
+		}
+	}
+	return false;
+}
+
+// Note, like DoPatch(FName&) we should strive to remove this method in favour of one that understands
+// the context for which this string belongs to. Patching it based on search and replace, is going to be 
+// error-prone and should be avoided.
 bool FAssetHeaderPatcherInner::DoPatch(FString& InOutString)
 {
+	// Attempt a direct replacement
 	{
 		// Find a Path, change a Path.
-		FStringView MaybeReplacement = Find(SearchAndReplace, InOutString);
+		FStringView MaybeReplacement = Find(StringReplacements, InOutString);
 		if (!MaybeReplacement.IsEmpty())
 		{
 			InOutString = MaybeReplacement;
@@ -599,155 +1394,369 @@ bool FAssetHeaderPatcherInner::DoPatch(FString& InOutString)
 		}
 	}
 
+	// Direct replacement failed so now try substring replacements
+
+	bool bDidPatch = false;
+	TStringBuilder<NAME_SIZE> DstStringBuilder;
 	{
-		// Patch Object paths.
+		// Patch Object paths with sub-object (not-necessarily quoted)
 		// Path occurs to the left of a ":"
-		int Idx{};
-		if (InOutString.FindChar(TCHAR(':'), Idx))
+		int32 ColonPos;
+		FStringView PathView(InOutString);
+		while (PathView.FindChar(SUBOBJECT_DELIMITER_CHAR, ColonPos))
 		{
-			if (InOutString[Idx + 1] == TCHAR(':'))
+			if ((ColonPos + 1) < PathView.Len() && PathView[ColonPos + 1] == SUBOBJECT_DELIMITER_CHAR)
 			{
 				// "::" is not a path delim
-				return false;
+				PathView.RightChopInline(ColonPos + 1);
+				continue;
 			}
-			FStringView MaybeReplacement = Find(SearchAndReplace, FStringView(*InOutString, Idx));
-			if (!MaybeReplacement.IsEmpty())
+
+			// Presumably we have found the start of a path's sub-object path. Create a new 
+			// view for our possible ObjectPath and walk backwards confirming we are in a path
+			// otherwise start over at the next ':'
+			FStringView ObjectPathView(PathView.GetData(), ColonPos);
+
+			int32 OuterDelimiterPos;
+			if (!ObjectPathView.FindLastChar(TEXT('.'), OuterDelimiterPos))
 			{
-				InOutString = MaybeReplacement + InOutString.RightChop(Idx);
-				return true;
+				// A ':' but '.' before it is not an object path
+				PathView.RightChopInline(ColonPos + 1);
+				continue;
 			}
-		}
-	}
 
-	{
-		// Patch quoted paths.
-		// Path occurs to the right of the first "'" 
-		int FirstQuotePos = INDEX_NONE;
-		if (InOutString.EndsWith(TEXT("'")) && InOutString.FindChar(TCHAR('\''), FirstQuotePos))
-		{
-			// +1 to skip first quote, +2 to account for the length minus the first and last quotes
-			FStringView StrippedQuotedPath = FStringView(*InOutString + FirstQuotePos + 1, InOutString.Len() - (FirstQuotePos + 2));
-			FStringView MaybeReplacement = Find(SearchAndReplace, StrippedQuotedPath);
+			int32 LastPathDelimiterPos = INDEX_NONE;
+			int32 Index = OuterDelimiterPos;
+			while (--Index >= 0)
+			{
+				if (ObjectPathView[Index] == TEXT('/'))
+				{
+					LastPathDelimiterPos = Index;
+				}
+				else
+				{
+					// Confirm we are still in a path
+					int32 PosInvalidChar = 0;
+					if (InvalidObjectPathCharacters.FindChar(ObjectPathView[Index], PosInvalidChar))
+					{
+						break;
+					}
+				}
+			}
+
+			if (LastPathDelimiterPos < 0)
+			{
+				// No '/' means we aren't in a path
+				PathView.RightChopInline(ColonPos + 1);
+				continue;
+			}
+
+			FStringView SrcMountPoint;
+			FStringView DstMountPoint;
+			FStringView ObjectPath(PathView.GetData() + LastPathDelimiterPos, ColonPos - LastPathDelimiterPos);
+			FStringView MaybeReplacement = Find(StringReplacements, ObjectPath);
 			if (!MaybeReplacement.IsEmpty())
 			{
-				InOutString.Reserve(FirstQuotePos + MaybeReplacement.Len() + 2); // + 2 for the two quotes
-				InOutString.LeftInline(FirstQuotePos + 1, EAllowShrinking::No); // Keep the original prefix before the quotes + the first quote
-				InOutString.Append(MaybeReplacement);
-				InOutString.AppendChar(TCHAR('\''));
-				return true;
+				FStringView LeftPart(*InOutString, int32(PathView.GetData() - *InOutString) + LastPathDelimiterPos);
+				FStringView RightPart(PathView.GetData() + ColonPos);
+
+				DstStringBuilder.Reset();
+				DstStringBuilder.Append(LeftPart);
+				DstStringBuilder.Append(MaybeReplacement);
+				DstStringBuilder.Append(RightPart);
+
+				InOutString = DstStringBuilder.ToString();
+				bDidPatch = true;
+
+				// Keep searching until the path is depleted since there might be more than one path to replace
+				PathView = FStringView(*InOutString + LeftPart.Len() + MaybeReplacement.Len() + 1);
+			}
+			else if (ShouldReplaceMountPoint(ObjectPath, SrcMountPoint, DstMountPoint))
+			{
+				FStringView LeftPart(*InOutString, int32(PathView.GetData() - *InOutString) + LastPathDelimiterPos);
+				FStringView RightPart(PathView.GetData() + LastPathDelimiterPos + SrcMountPoint.Len());
+
+				DstStringBuilder.Reset();
+				DstStringBuilder.Append(LeftPart);
+				DstStringBuilder.Append(DstMountPoint);
+				DstStringBuilder.Append(RightPart);
+
+				InOutString = DstStringBuilder.ToString();
+				bDidPatch = true;
+
+				// Keep searching until the path is depleted since there might be more than one path to replace
+				// Skip to the colon since we know we didn't have any matches within the quotes beyond the mount
+				PathView = FStringView(*InOutString + ColonPos + 1);
+			}
+			else
+			{
+				// No match but keep searching as there may be more than one ':'
+				PathView.RightChopInline(ColonPos + 1);
 			}
 		}
 	}
 
 	{
-		// Patching embedded mount points.
-		// 
-		FStringView InView = InOutString;
-		FStringView MountPoint = FPathViews::GetMountPointNameFromPath(InView);
-		if (MountPoint.Len())
-		{
-			TStringBuilder<NAME_SIZE> Tmp;
-			Tmp.Append(TEXT("<Mountpoint>"));
-			Tmp.Append(MountPoint);
-			FStringView MaybeReplacement = Find(SearchAndReplace, FStringView(Tmp));
-			if (!MaybeReplacement.IsEmpty())
+		// Patch quoted paths.
+		// Path occurs to the right of the first "'" or """ 
+		auto PatchQuotedPath = [this, &DstStringBuilder](FString& StringToPatch, FStringView Quote)
 			{
-				int32 Left = UE_PTRDIFF_TO_INT32(MountPoint.GetData() - InView.GetData());
-				int32 Right = Left + MountPoint.Len();
+				int32 FirstQuotePos = INDEX_NONE;
+				bool bFoundReplacement = false;
+				FStringView PathView(StringToPatch);
+				while ((FirstQuotePos = PathView.Find(Quote, 0, ESearchCase::CaseSensitive)) != INDEX_NONE)
+				{
+					int32 SecondQuotePos = PathView.Find(Quote, FirstQuotePos + 1, ESearchCase::CaseSensitive);
+					if (SecondQuotePos == INDEX_NONE)
+					{
+						// If there isn't a second quote we're done
+						break;
+					}
+
+					FStringView SrcMountPoint;
+					FStringView DstMountPoint;
+					FStringView StrippedQuotedPath = FStringView(PathView.GetData() + FirstQuotePos + 1, SecondQuotePos - FirstQuotePos - 1); // +1 and -1 are to skip the quotes
+					FStringView MaybeReplacement = Find(StringReplacements, StrippedQuotedPath);
+					if (!MaybeReplacement.IsEmpty())
+					{
+						FStringView LeftPart(*StringToPatch, int32(PathView.GetData() - *StringToPatch) + FirstQuotePos + 1); // +1 to ensure we include the quote
+						FStringView RightPart(PathView.GetData() + SecondQuotePos);
+
+						DstStringBuilder.Reset();
+						DstStringBuilder.Append(LeftPart);
+						DstStringBuilder.Append(MaybeReplacement);
+						DstStringBuilder.Append(RightPart);
+
+						StringToPatch = DstStringBuilder.ToString();
+						bFoundReplacement = true;
+
+						// Keep searching until the path is depleted since there might be more than one path to replace
+						PathView = FStringView(*StringToPatch + LeftPart.Len() + MaybeReplacement.Len() + 1);
+					}
+					else if (ShouldReplaceMountPoint(StrippedQuotedPath, SrcMountPoint, DstMountPoint))
+					{
+						FStringView LeftPart(*StringToPatch, int32(PathView.GetData() - *StringToPatch) + FirstQuotePos + 1); // +1 to ensure we include the quote
+						FStringView RightPart(PathView.GetData() + FirstQuotePos + SrcMountPoint.Len() + 1); // +1 to ensure we skip the first quote
+
+						DstStringBuilder.Reset();
+						DstStringBuilder.Append(LeftPart);
+						DstStringBuilder.Append(DstMountPoint);
+						DstStringBuilder.Append(RightPart);
+
+						StringToPatch = DstStringBuilder.ToString();
+						bFoundReplacement = true;
+
+						// Keep searching until the path is depleted since there might be more than one path to replace
+						// Skip to the end quote since we know we didn't have any matches within the quotes beyond the mount
+						PathView = FStringView(*StringToPatch + SecondQuotePos + 1);
+					}
+					else
+					{
+						// No match but keep searching as there may be more than one quoted path
+						PathView.RightChopInline(SecondQuotePos + 1);
+					}
+				}
+				return bFoundReplacement;
+			};
+		bDidPatch |= PatchQuotedPath(InOutString, TEXT("'"));
+		bDidPatch |= PatchQuotedPath(InOutString, TEXT("\""));
+	}
 
-				Tmp.Reset();
-				Tmp.Append(InView.Left(Left));
-				Tmp.Append(MaybeReplacement);
-				Tmp.Append(InView.RightChop(Right));
+	return bDidPatch;
+}
 
-				InOutString = Tmp;
-				return true;
-			}
+bool FAssetHeaderPatcherInner::RemapFName(FName SrcName, FName DstName)
+{
+	// NameTable entries only care about the comparison form (no number) so 
+	// only consider that for remapping purposes
+	FNameEntryId SrcComparisonId = SrcName.GetDisplayIndex();
+	FNameEntryId DstComparisonId = DstName.GetDisplayIndex();
+	if (SrcComparisonId == DstComparisonId)
+	{
+		return false;
+	}
+	checkf(DstName != NAME_None, TEXT("There should never be a None FName in the NameTable"));
+
+	FNameEntryId* RemappedFName = RenameMap.Find(SrcComparisonId);
+	if (RemappedFName)
+	{
+		// We already have a mapping. That is fine; we might have used the same FName in more than one place.
+		// However, we need to be certain we are renaming the name to the same new name. If not, this means
+		// the originals names overlapped but in the patched case they don't (e.g. A class FName may have matched a Package
+		// name, but after patching it's possible _only_ the Package name has changed. In such a case we don't want to rename the 
+		// class name inadvertently by patching the shared NameTable entry. If we have a mismatch with the new patched 
+		// name, record the new name and we will append it to the NameTable later.
+		if (*RemappedFName != DstComparisonId)
+		{
+			AddedNames.Add(DstComparisonId);
 		}
 	}
+	else
+	{
+		RenameMap.Add(SrcComparisonId, DstComparisonId);
+	}
 
-	return false;
+	return true;
 }
 
 bool FAssetHeaderPatcherInner::DoPatch(FName& InOutName)
 {
-	// FNames are first compared with as a ToString which may include a `_[0-9+]` tail.
-	// This is because FNames can be used as a file name can have these characters.
-	// FNames are then compared without this number if one is present as they are used to link objects
-	// within the object domain and the migration being performed may not care about the number.
+	// If we are given an FName to patch we have no real context as to what that FName is
+	// so we conservatively assume it is a package path and attempt to patch that only
+	FCoreRedirectObjectName SrcPackageName(NAME_None, NAME_None, InOutName);
+	FCoreRedirectObjectName DstPackageName = FCoreRedirects::GetRedirectedName(ECoreRedirectFlags::Type_Package, SrcPackageName);
+	if (RemapFName(SrcPackageName.PackageName, DstPackageName.PackageName))
+	{
+		InOutName = DstPackageName.PackageName;
+		return true;
+	}
+	return false;
+}
+
+void FAssetHeaderPatcherInner::PatchNameTable()
+{
+	// Note, no number is assigned when replacing FNames as the NameTable only tracks unnumbered names
+
+	// Update the NameTable with the known patched values and add our new patched names to the NameToIndex
+	// map so we can validate that we always have a FName mapping to an entry in the name table when writing
+	for (auto& Pair : RenameMap)
+	{
+		FNameEntryId SrcName = Pair.Key;
+		FNameEntryId DstName = Pair.Value;
+		int32* pSrcIndex = NameToIndexMap.Find(SrcName);
+		checkf(pSrcIndex && *pSrcIndex < NameTable.Num(), TEXT("An FName remapping was done for a name (%s) not in the NameTable."), *FName::CreateFromDisplayId(DstName, NAME_NO_NUMBER_INTERNAL).ToString());
+		int32 SrcIndex = *pSrcIndex;
+
+		NameTable[SrcIndex] = FName::CreateFromDisplayId(DstName, NAME_NO_NUMBER_INTERNAL);
+		NameToIndexMap.Add(DstName, SrcIndex);
+	}
+
+	for (FNameEntryId NewName : AddedNames)
+	{
+		FName NewFName = FName::CreateFromDisplayId(NewName, NAME_NO_NUMBER_INTERNAL);
+		int32 NameTableIndex = NameTable.Num();
 
-	// If a new value it is important that the new FName is created in a deterministic way.
-	// So, if the new value has a numeric tail, this ios used.
-	// If not, then the original number is used.
+		NameTable.Add(NewFName);
+		NameToIndexMap.Add(NewFName.GetDisplayIndex(), NameTableIndex);
+	}
 
-	bool bFound = false;
-	FString Value = InOutName.ToString();
-	bFound = DoPatch(Value);
+	Summary.NameCount = NameTable.Num();
+}
 
-	if (!bFound && InOutName.GetNumber() != 0)
+bool FAssetHeaderPatcherInner::DoPatch(FSoftObjectPath& InOutSoft)
+{
+	FTopLevelAssetPath InOutTopLevelAssetPath = InOutSoft.GetAssetPath();
+	if (!DoPatch(InOutTopLevelAssetPath))
 	{
-		// We didn't find it, and it had a number, so try without the number.
-		Value = InOutName.GetPlainNameString();
-		bFound = DoPatch(Value);
+		return false;
 	}
 
-	if (bFound)
+	InOutSoft.SetPath(InOutTopLevelAssetPath, InOutSoft.GetSubPathString());
+	return true;
+}
+
+bool FAssetHeaderPatcherInner::DoPatch(FObjectResource& InOutResource, bool bIsExport, FName& OutPatchedObjectName)
+{
+	bool bOutermostIsExport = bIsExport;
+	FPackageIndex OuterIndex = InOutResource.OuterIndex;
+	TArray<FName, TInlineAllocator<8>> OuterStack;
+	while (!OuterIndex.IsNull())
 	{
-		// If the string comes with its own number, preserve it.
-		// If not preserve the original number.
-		FName Tmp(*Value);
-		if (Tmp.GetDisplayNameEntry()->GetNameLength() != Value.Len())
+		const FObjectResource* OuterResource;
+		if (OuterIndex.IsImport())
 		{
-			// Value had a '_[0-9]+' tail.
-			InOutName = Tmp;
+			bOutermostIsExport = false;
+			OuterResource = &ImportTable[OuterIndex.ToImport()];
 		}
 		else
 		{
-			// Value did not have a tail, so use the original number.
-			InOutName = FName(Value, InOutName.GetNumber());
+			bOutermostIsExport = true;
+			OuterResource = &ExportTable[OuterIndex.ToExport()];
 		}
 
-		// For specific cases we want the number to be parsed in the FName.
-		// This is usually when we use it as a path.
-		// It may be, that dont have a FName in the table for that case at this point, so we add it.
-		// This will probably contribute to unused FNames in the FName table, 
-		// but these will most likely be removed on a subsequent full save.
-		if (!NameToIndexMap.Contains(InOutName.GetDisplayIndex()))
+		OuterStack.Push(OuterResource->ObjectName);
+		OuterIndex = OuterResource->OuterIndex;
+	}
+
+	FName SrcObjectName;
+	FName SrcOuterName;
+	FName SrcPackageName;
+	bool bRemapByPackageName = false;
+	if (OuterStack.Num() == 0)
+	{
+		if (bOutermostIsExport)
 		{
-			int32 Indx = NameTable.Num();
-			NameTable.Add(InOutName);
-			NameToIndexMap.Add(InOutName.GetDisplayIndex(), Indx);
+			SrcPackageName = OriginalPackagePath;		// /Package/Package
+			SrcOuterName = NAME_None;
+			SrcObjectName = InOutResource.ObjectName;	// MyObject
+		}
+		else
+		{
+			// The ObjectName is a package
+			SrcPackageName = InOutResource.ObjectName;	// /Package/Package
+			SrcOuterName = NAME_None;
+			SrcObjectName = NAME_None;
+			bRemapByPackageName = true;
 		}
-		return true;
 	}
-	return false;
-}
+	else
+	{
+		SrcPackageName = bOutermostIsExport ? OriginalPackagePath : OuterStack.Pop();
 
-bool FAssetHeaderPatcherInner::DoPatch(FSoftObjectPath& InOutSoft)
-{
-	FTopLevelAssetPath TmpPath = InOutSoft.GetAssetPath();
-	if (DoPatch(TmpPath))
+		TStringBuilder<NAME_SIZE> OuterString;
+		while (!OuterStack.IsEmpty())
+		{
+			FName Outer = OuterStack.Pop();
+			Outer.ToString(OuterString);
+			OuterString.AppendChar(TEXT('.'));
+		}
+		if (OuterString.Len())
+		{
+			OuterString.RemoveSuffix(1);
+		}
+		SrcOuterName = FName(OuterString);
+		SrcObjectName = InOutResource.ObjectName;
+	}
+	
+	const FCoreRedirectObjectName SrcObjectPath(SrcObjectName, SrcOuterName, SrcPackageName);
+	const FCoreRedirectObjectName DstObjectPath = FCoreRedirects::GetRedirectedName(ECoreRedirectFlags::Type_AllMask, SrcObjectPath);
+
+	bool bPatched = false;
+	if (!bRemapByPackageName)
 	{
-		InOutSoft.SetPath(TmpPath, InOutSoft.GetSubPathString());
-		return true;
+		bPatched = RemapFName(SrcObjectPath.ObjectName, DstObjectPath.ObjectName);
+		OutPatchedObjectName = DstObjectPath.ObjectName;
 	}
-	return false;
+	else
+	{
+		bPatched = RemapFName(SrcObjectPath.PackageName, DstObjectPath.PackageName);
+		OutPatchedObjectName = DstObjectPath.PackageName;
+	}
+
+#if WITH_EDITORONLY_DATA
+	InOutResource.OldClassName = NAME_None;
+#endif
+
+	return bPatched;
 }
 
 bool FAssetHeaderPatcherInner::DoPatch(FTopLevelAssetPath& InOutPath)
 {
-	FName PackageName = InOutPath.GetPackageName();
-	FName AssetName = InOutPath.GetAssetName();
-	bool bPatching = DoPatch(PackageName);
-	bPatching |= DoPatch(AssetName);
-	if (bPatching)
+	const FCoreRedirectObjectName SrcTopLevelAssetPath = FCoreRedirects::GetRedirectedName(ECoreRedirectFlags::Type_AllMask, InOutPath);
+	const FTopLevelAssetPath DstTopLevelAssetPath(SrcTopLevelAssetPath.ToString());
+
+	bool bPatched = RemapFName(InOutPath.GetAssetName(), DstTopLevelAssetPath.GetAssetName());
+	bPatched |= RemapFName(InOutPath.GetPackageName(), DstTopLevelAssetPath.GetPackageName());
+
+	if (bPatched)
 	{
-		InOutPath = FTopLevelAssetPath(PackageName, AssetName);
+		InOutPath = DstTopLevelAssetPath;
 	}
-	return bPatching;
+
+	return bPatched;
 }
 
-bool FAssetHeaderPatcherInner::DoPatch(FGatherableTextData& InGatherableTextData)
+bool FAssetHeaderPatcherInner::DoPatch(FGatherableTextData& InOutGatherableTextData)
 {
 	// There are various fields in FGatherableTextData however only one pertains to 
 	// asset paths and types, SourceSiteContexts.SiteDescription. The rest are contextual
@@ -756,7 +1765,7 @@ bool FAssetHeaderPatcherInner::DoPatch(FGatherableTextData& InGatherableTextData
 	// may try to load from these strings)
 
 	bool bDidPatch = false;
-	for (FTextSourceSiteContext& SourceSiteContext : InGatherableTextData.SourceSiteContexts)
+	for (FTextSourceSiteContext& SourceSiteContext : InOutGatherableTextData.SourceSiteContexts)
 	{
 		FStringView ClassName;
 		FStringView PackagePath;
@@ -764,174 +1773,257 @@ bool FAssetHeaderPatcherInner::DoPatch(FGatherableTextData& InGatherableTextData
 		FStringView SubObjectName;
 		FPackageName::SplitFullObjectPath(SourceSiteContext.SiteDescription, ClassName, PackagePath, ObjectName, SubObjectName, true /*bDetectClassName*/);
 
-		// If we don't have a PackagePath, do nothing. Replacing object names alone without the package path is likely going to lead to incorrect replacements
-		if (PackagePath.IsEmpty())
+		// Todo to use StringView logic above to reduce string copies
+		FSoftObjectPath SiteDescriptionPath(SourceSiteContext.SiteDescription);
+		if (!SiteDescriptionPath.IsValid())
 		{
 			continue;
 		}
 
-		const FString* DstPackagePath = SearchAndReplace.FindByHash(GetTypeHash(PackagePath), PackagePath);
-		if (!DstPackagePath)
+		FTopLevelAssetPath TopLevelAssetPath = SiteDescriptionPath.GetAssetPath();
+		const FCoreRedirectObjectName RedirectedTopLevelAssetPath = FCoreRedirects::GetRedirectedName(ECoreRedirectFlags::Type_AllMask, TopLevelAssetPath);
+		const FTopLevelAssetPath PatchedTopLevelAssetPath(RedirectedTopLevelAssetPath.ToString());
+		if (TopLevelAssetPath == PatchedTopLevelAssetPath)
 		{
 			continue;
 		}
-
-		// Note, we don't provide remappings down to the sub-object level so we can skip checking for and handling of sub-object specific renames
-		// We also don't bother stripping off the ClassName below, nor do we restore a prepended classname in the Dst string by scanning the replacement string
-		// if one existed in the Src string. I don't believe SiteDescriptions will ever contain a ClassName, but just in case I'm wrong we'll add a check.
-		check(ClassName.IsEmpty());
-
 		bDidPatch = true;
-
-		if (ObjectName.IsEmpty())
-		{
-			SourceSiteContext.SiteDescription = *DstPackagePath;
-			continue;
-		}
-
-		// If we have a object the whole path can be used (minus SubObjectName )
-		FStringView FullObjectPath(SourceSiteContext.SiteDescription);
-		FullObjectPath.LeftChopInline(SubObjectName.Len() + !!SubObjectName.Len()); // Strip off the SubObjectName + 1 for the ':' separator (if we have a subobject name at all)
-
-		if (const FString* DstFullObjectPath = SearchAndReplace.FindByHash(GetTypeHash(FullObjectPath), FullObjectPath))
-		{
-			// We have a replacement for the PackageName.ObjectName but not the subobject, so append the subobject if necessary
-			if (!SubObjectName.IsEmpty())
-			{
-				// We are writing over the original string, so capture our stringview before proceeding
-				FString LocalSubObjectName(SubObjectName);
-				SourceSiteContext.SiteDescription.Reserve(DstFullObjectPath->Len() + SubObjectName.Len() + 1); // + 1 for ':'
-				SourceSiteContext.SiteDescription = *DstFullObjectPath;
-				SourceSiteContext.SiteDescription.AppendChar(TEXT(':'));
-				SourceSiteContext.SiteDescription.Append(LocalSubObjectName);
-			}
-			else
-			{
-				SourceSiteContext.SiteDescription = *DstFullObjectPath;
-			}
-		}
-		else
-		{
-			// We only have a PackagePath replacement so only replace that portion of the string and keep the object name and sub-object paths untouched.
-			// Take all text from the start of the ObjectName until the end of SiteDescription
-			// This will let us include the subobject path if any since we know SiteDescription is null-terminated
-			
-			// We are writing over the original string, so capture our stringview before proceeding
-			FString SrcFullObjectPath(ObjectName.GetData());
-			SourceSiteContext.SiteDescription.Reserve(DstPackagePath->Len() + SrcFullObjectPath.Len() + 1); // + 1 for '.'
-			SourceSiteContext.SiteDescription = *DstPackagePath;
-			SourceSiteContext.SiteDescription.AppendChar(TEXT('.'));
-			SourceSiteContext.SiteDescription.Append(SrcFullObjectPath);
-		}
+		SiteDescriptionPath.SetPath(PatchedTopLevelAssetPath, SiteDescriptionPath.GetSubPathString());
+		SourceSiteContext.SiteDescription = SiteDescriptionPath.ToString();
 	}
 	
 	return bDidPatch;
 }
 
-void FAssetHeaderPatcherInner::PatchHeader_PatchSections()
+bool FAssetHeaderPatcherInner::DoPatch(FThumbnailEntry& InThumbnailEntry)
 {
-	DoPatch(Summary.PackageName);
-
-	{	// Patch the Name table
-		// This data is used by all FNames in the file.
-		// So if we patch a Identifier we append it to the name table.
-		// This is because there may be FNames in data we dont want to patch in structures we dont look at.
-		// If we patch the ident at inplace, then we would change those names.
-		TArray<FName> ToAppend;
-		for (FName& Name : NameTable)
-		{
-			FNameEntryId OriginalFNameEntryId = Name.GetDisplayIndex();
-			FString TmpNameStr = Name.ToString(); // will contain a "[0-9]+" suffix, if any
+	// These objects can potentially be paths to sub-objects. For renaming purposes we 
+	// want to drop the sub-object path and grab the AssetName
+	FStringView SrcObjectPathWithoutPackageName(InThumbnailEntry.ObjectPathWithoutPackageName);
+	int32 ColonPos = INDEX_NONE;
+	if (SrcObjectPathWithoutPackageName.FindChar(TEXT(':'), ColonPos))
+	{
+		SrcObjectPathWithoutPackageName.LeftChopInline(SrcObjectPathWithoutPackageName.Len() - ColonPos);
+	}
 
-			if (DoPatch(TmpNameStr))
-			{
-				// Construct a new FName based on our patched string.
-				// If the string has a "[0-9]+" suffix we keep it as part of the plain string name.
-				// When we write the NameTable back out, it will serialize the string with the number, but not the number set here.
-				FName TmpName = FName(*TmpNameStr, NAME_NO_NUMBER);
+	FName PackageFName = OriginalPackagePath;
+	if (bIsNonOneFilePerActorPackage)
+	{
+		PackageFName = OriginalNonOneFilePerActorPackagePath;
+	}
 
-				// If the string does not contain any path separators, then call it an Identifier.
-				if (!(TmpNameStr.Contains(TEXT("/")) || TmpNameStr.Contains(TEXT("\\"))))
-				{
-					ToAppend.Add(TmpName);
-				}
-				else
-				{
-					Name = TmpName;
-				}
-			}
-		}
-		NameTable.Append(ToAppend);
+	const FCoreRedirectObjectName SrcTopLevelAssetName(FName(SrcObjectPathWithoutPackageName), NAME_None, PackageFName);
+	const FCoreRedirectObjectName RedirectedTopLevelAssetName = FCoreRedirects::GetRedirectedName(ECoreRedirectFlags::Type_Object, SrcTopLevelAssetName);
+	bool bPatched = RemapFName(SrcTopLevelAssetName.ObjectName, RedirectedTopLevelAssetName.ObjectName);
 
-		Summary.NameCount = NameTable.Num();
-	}
+	const FCoreRedirectObjectName SrcClassName(FName(InThumbnailEntry.ObjectShortClassName), NAME_None, NAME_None);
+	const FCoreRedirectObjectName RedirectedClassName = FCoreRedirects::GetRedirectedName(ECoreRedirectFlags::Type_Class, SrcClassName);
+	bPatched |= RemapFName(SrcClassName.ObjectName, RedirectedClassName.ObjectName);
 
-	// Build Name Path table
-	// This is used to generate the indices saved by the FNamePatchingWriter
-	NameToIndexMap.Reserve(NameTable.Num());
-	for (TConstEnumerateRef<FName> Name : EnumerateRange(NameTable))
+	if (bPatched)
 	{
-		NameToIndexMap.Add(Name->GetDisplayIndex()) = Name.GetIndex();
+		// Since we patched, we will cause the inline string name to affect the thumbnail offsets. Calculate the size change
+		// here so we can use it during writing where we will fix up the offsets.
+		int32 Delta = -(InThumbnailEntry.ObjectShortClassName.Len() + InThumbnailEntry.ObjectPathWithoutPackageName.Len());
+
+		InThumbnailEntry.ObjectShortClassName = RedirectedClassName.ObjectName.ToString();
+		InThumbnailEntry.ObjectPathWithoutPackageName = RedirectedTopLevelAssetName.ObjectName.ToString();
+		Delta += InThumbnailEntry.ObjectShortClassName.Len() + InThumbnailEntry.ObjectPathWithoutPackageName.Len();
+
+		InThumbnailEntry.Delta = Delta;
 	}
 
-	// Soft paths
-	for (FSoftObjectPath& PathRef : SoftObjectPathTable)
+	return bPatched;
+}
+
+void FAssetHeaderPatcherInner::PatchHeader_PatchSections()
+{
+	// Package Summary
 	{
-		DoPatch(PathRef);
+		const FCoreRedirectObjectName DstPackageName = FCoreRedirects::GetRedirectedName(ECoreRedirectFlags::Type_Package,
+			FCoreRedirectObjectName(NAME_None, NAME_None, OriginalPackagePath));
+
+		// This is a string, so we do not want to Remap the patched name unless it's a non-OFPA
+		// package, in which case there will be an FName entry for this path
+		Summary.PackageName = DstPackageName.PackageName.ToString();
+		
+		// It seems that non-OFPA packages tend to have the package name in the nametable, 
+		// however it isn't a guarantee, so we confirm the name is there before remapping and
+		// extend this special case of NameTable patching to all packages, OFPA or not.
+		if (NameToIndexMap.Find(OriginalPackagePath.GetDisplayIndex()))
+		{
+			RemapFName(OriginalPackagePath, DstPackageName.PackageName);
+		}
 	}
 
-	// GatherableTextData table
-	for (FGatherableTextData& GatherableTextData : GatherableTextDataTable)
+	// Patching of the FObjectResource ObjectNames is deferred since when patching we need to 
+	// walk the original names to determine if they need patching in the first place
+	TMap<int32, FName> PatchedExportObjectNames;
+	TMap<int32, FName> PatchedImportObjectNames;
+
+	// Export Table
 	{
-		DoPatch(GatherableTextData);
+		PatchedExportObjectNames.Reserve(ExportTable.Num());
+		for (int32 i = 0; i < ExportTable.Num(); ++i)
+		{
+			FObjectExport& Export = ExportTable[i];
+			FName PatchedObjectName;
+			if (DoPatch(Export, true, PatchedObjectName))
+			{
+				PatchedExportObjectNames.Add(i, PatchedObjectName);
+			}
+		}
 	}
 
 	// Import table
-	for (FObjectImport& Import : ImportTable)
 	{
-		DoPatch(Import.ObjectName);
-#if WITH_EDITORONLY_DATA
-		DoPatch(Import.OldClassName);
-#endif
+		PatchedImportObjectNames.Reserve(ImportTable.Num());
+		for (int32 i = 0; i < ImportTable.Num(); ++i)
+		{
+			FObjectImport& Import = ImportTable[i];
+			FName PatchedObjectName;
+			if (DoPatch(Import, false, PatchedObjectName))
+			{
+				PatchedImportObjectNames.Add(i, PatchedObjectName);
+			}
+
+			const FCoreRedirectObjectName SrcClass(Import.ClassName, NAME_None, Import.ClassPackage);
+			const FCoreRedirectObjectName DstClass = FCoreRedirects::GetRedirectedName(ECoreRedirectFlags::Type_Class | ECoreRedirectFlags::Type_Package, SrcClass);
+
+			if (RemapFName(SrcClass.ObjectName, DstClass.ObjectName))
+			{
+				Import.ClassName = DstClass.ObjectName;
+			}
+
+			if (RemapFName(SrcClass.PackageName, DstClass.PackageName))
+			{
+				Import.ClassPackage = DstClass.PackageName;
+			}
 
-		DoPatch(Import.ClassPackage);
-		DoPatch(Import.ClassName);
 #if WITH_EDITORONLY_DATA
-		DoPatch(Import.PackageName);
+			const FCoreRedirectObjectName SrcPackage(NAME_None, NAME_None, Import.PackageName);
+			const FCoreRedirectObjectName DstPackage = FCoreRedirects::GetRedirectedName(ECoreRedirectFlags::Type_Package, SrcPackage);
+
+			if (RemapFName(SrcPackage.PackageName, DstPackage.PackageName))
+			{
+				Import.PackageName = DstPackage.PackageName;
+			}
 #endif
+		}
 	}
 
-	// Export Table
-	for (FObjectExport& Export : ExportTable)
+	// Finish the FObjectResource patching that was deferred above for the ExportTable and ImportTable
 	{
-		DoPatch(Export.ObjectName);
-#if WITH_EDITORONLY_DATA
-		DoPatch(Export.OldClassName);
-#endif
+		for (auto& Pair : PatchedExportObjectNames)
+		{
+			int32 Index = Pair.Key;
+			FName PatchedName = Pair.Value;
+			ExportTable[Index].ObjectName = PatchedName;
+		}
+
+		for (auto& Pair : PatchedImportObjectNames)
+		{
+			int32 Index = Pair.Key;
+			FName PatchedName = Pair.Value;
+			ImportTable[Index].ObjectName = PatchedName;
+		}
+	}
+
+	// Soft paths
+	for (FSoftObjectPath& SoftObjectPath : SoftObjectPathTable)
+	{
+		DoPatch(SoftObjectPath);
+	}
+
+	// GatherableTextData table
+	for (FGatherableTextData& GatherableTextData : GatherableTextDataTable)
+	{
+		DoPatch(GatherableTextData);
 	}
 
 	// Soft Package Reference's
 	for (FName& Reference : SoftPackageReferencesTable)
 	{
-		DoPatch(Reference);
+		FCoreRedirectObjectName SrcPackagePath(NAME_None, NAME_None, Reference);
+		FCoreRedirectObjectName DstPackageName = FCoreRedirects::GetRedirectedName(ECoreRedirectFlags::Type_Package, SrcPackagePath);
+		if (RemapFName(SrcPackagePath.PackageName, DstPackageName.PackageName))
+		{
+			Reference = DstPackageName.PackageName;
+		}
 	}
 
+/*
+	Enable when serialization is replaced with PackageReader. Thumbnail serialization seems to work but 
+	we don't test it well enough to confidently enable. This data is regenerated when missing upon load 
+	so not patching it isn't a huge deal. Patching it incorrectly may corrupt the whole file however.
+
+	// SearchableNamesMap
+	for (auto& Pair : SearchableNamesMap)
+	{
+		TArray<FName>& Names = Pair.Value;
+		for (FName& Name : Names)
+		{
+			DoPatch(Name);
+		}
+	}
+
+	// Thumbnail Table
+	for (FThumbnailEntry& ThumbnailEntry : ThumbnailTable)
+	{
+		DoPatch(ThumbnailEntry);
+	}
+*/
+
 	// Asset Register Data
 	for (FAssetRegistryObjectData& ObjData : AssetRegistryData.ObjectData)
 	{
-		DoPatch(ObjData.ObjectData.ObjectPath);
+		// ObjectPath is a toss-up. 
+		// Sometimes it's a FTopLevelAssetPath with an implied PackageName (this package's name) and AssetName.
+		// Sometimes it's a full FSoftPath (e.g. when dealing with ExternalObjects)
+		FSoftObjectPath SrcObjectPath(ObjData.ObjectData.ObjectPath);
+		{
+			if (SrcObjectPath.IsValid())
+			{
+				FSoftObjectPath SrdDstObjectPath = SrcObjectPath;
+
+				if (DoPatch(SrdDstObjectPath))
+				{
+					ObjData.ObjectData.ObjectPath = SrdDstObjectPath.ToString();
+				}
+			}
+			else
+			{
+				FTopLevelAssetPath SrcDstTopLevelAssetPath(OriginalPackagePath, FName(ObjData.ObjectData.ObjectPath));
+				SrcObjectPath.SetPath(SrcDstTopLevelAssetPath, SrcObjectPath.GetSubPathString());
+
+				if (DoPatch(SrcDstTopLevelAssetPath))
+				{
+					ObjData.ObjectData.ObjectPath = SrcDstTopLevelAssetPath.GetAssetName().ToString();
+				}
+			}
+		}
+
+		// ObjData.ObjectData.ObjectClassName is a FTopLevelAssetPath stored as a string
+		FTopLevelAssetPath SrcObjectClassName(ObjData.ObjectData.ObjectClassName);
+		{
+			FTopLevelAssetPath SrcDstObjectClassName = SrcObjectClassName;
 
-		{	// ObjData.ObjectData.ObjectClassName is stored as a string, but is used as a FTopLevelAssetPath
-			FTopLevelAssetPath Tmp(ObjData.ObjectData.ObjectClassName);
-			if (DoPatch(Tmp))
+			if (DoPatch(SrcDstObjectClassName))
 			{
-				FString TS = Tmp.ToString();
-				ObjData.ObjectData.ObjectClassName = TS;
+				ObjData.ObjectData.ObjectClassName = SrcDstObjectClassName.ToString();
 			}
 		}
 
 		for (UE::AssetRegistry::FDeserializeTagData& TagData : ObjData.TagData)
 		{
+			if (IgnoredTags.Contains(TagData.Key))
+			{
+				continue;
+			}
+
+			// WorldPartitionActor metadata is special. It's an encoded string blob which needs
+			// handling internally, so we make use of a custom patcher to let us intercept
+			// various elements that might need patching.
 			if (TagData.Key == FWorldPartitionActorDescUtils::ActorMetaDataTagName())
 			{
 				const FString LongPackageName(SrcAsset);
@@ -943,10 +2035,32 @@ void FAssetHeaderPatcherInner::PatchHeader_PatchSections()
 				struct FWorldPartitionAssetDataPatcherInner : FWorldPartitionAssetDataPatcher
 				{
 					FWorldPartitionAssetDataPatcherInner(FAssetHeaderPatcherInner* InInner) : Inner(InInner) {}
-					virtual bool DoPatch(FString& InOutString) override { return Inner->DoPatch(InOutString); }
-					virtual bool DoPatch(FName& InOutName) override { return Inner->DoPatch(InOutName); }
-					virtual bool DoPatch(FSoftObjectPath& InOutSoft) override { return Inner->DoPatch(InOutSoft); }
-					virtual bool DoPatch(FTopLevelAssetPath& InOutPath) override { return Inner->DoPatch(InOutPath); }
+					virtual bool DoPatch(FString& InOutString) override 
+					{ 
+						return Inner->DoPatch(InOutString); 
+					}
+					virtual bool DoPatch(FName& InOutName) override 
+					{ 
+						// FNames are actually strings inside WorldPartitionActor metadata, and since a lone
+						// FName has no context for how to patch it, convert it to a string to perform a 
+						// best-effort search.
+						FString NameString;
+						InOutName.ToString(NameString);
+						if (Inner->DoPatch(NameString))
+						{
+							InOutName = FName(NameString);
+							return true;
+						}
+						return false;
+					}
+					virtual bool DoPatch(FSoftObjectPath& InOutSoft) override 
+					{
+						return Inner->DoPatch(InOutSoft);
+					}
+					virtual bool DoPatch(FTopLevelAssetPath& InOutPath) override 
+					{ 
+						return Inner->DoPatch(InOutPath);
+					}
 					FAssetHeaderPatcherInner* Inner;
 				};
 
@@ -957,14 +2071,26 @@ void FAssetHeaderPatcherInner::PatchHeader_PatchSections()
 					TagData.Value = PatchedAssetData;
 				}
 			}
-			else if (bBespokeSearchInUse &&
-				FPathViews::GetBaseFilename(ObjData.ObjectData.ObjectPath) == TEXT("GameFeatureData") &&
-				TagData.Key == TEXT("PrimaryAssetName"))
+			// Special case for common Tag
+			else if (bPatchPrimaryAssetTag && TagData.Key == TEXT("PrimaryAssetName"))
 			{
-				FString BespokeSearchValue = TEXT("<GameFeatureData.PrimaryAssetName>") + TagData.Value;
-				if (DoPatch(BespokeSearchValue))
+				if (TagData.Value == OriginalPrimaryAssetName)
 				{
-					TagData.Value = BespokeSearchValue;
+					const FCoreRedirectObjectName DstPackageName = FCoreRedirects::GetRedirectedName(ECoreRedirectFlags::Type_Package,
+						FCoreRedirectObjectName(NAME_None, NAME_None, OriginalPackagePath));
+
+					TStringBuilder<256> Builder;
+					DstPackageName.PackageName.ToString(Builder);
+					FStringView PrimaryAssetView = Builder.ToView();
+					ensure(PrimaryAssetView.Len() && PrimaryAssetView[0] == TEXT('/'));
+					PrimaryAssetView.RemovePrefix(1);
+
+					int32 SlashPos = INDEX_NONE;
+					if (PrimaryAssetView.FindChar(TEXT('/'), SlashPos))
+					{
+						TagData.Value.Empty();
+						TagData.Value.Append(PrimaryAssetView.GetData(), SlashPos);
+					}
 				}
 			}
 			else
@@ -974,8 +2100,8 @@ void FAssetHeaderPatcherInner::PatchHeader_PatchSections()
 		}
 	}
 
-	// Just incase new names were added due to number parsing
-	Summary.NameCount = NameTable.Num();
+	// Do nametable patching last since we want to ensure we have determined all the remappings necessary
+	PatchNameTable();
 }
 
 FAssetHeaderPatcher::EResult FAssetHeaderPatcherInner::PatchHeader_WriteDestinationFile()
@@ -991,6 +2117,7 @@ FAssetHeaderPatcher::EResult FAssetHeaderPatcherInner::PatchHeader_WriteDestinat
 		{ EPatchedSection::ImportTable,					Summary.ImportOffset,					HeaderInformation.ImportTableSize,					true	},
 		{ EPatchedSection::ExportTable,					Summary.ExportOffset,					HeaderInformation.ExportTableSize,					true	},
 		{ EPatchedSection::SoftPackageReferencesTable,	Summary.SoftPackageReferencesOffset,	HeaderInformation.SoftPackageReferencesListSize,	false	},
+		{ EPatchedSection::SearchableNamesMap,			Summary.SearchableNamesOffset,			HeaderInformation.SearchableNamesMapSize,			false	},
 		{ EPatchedSection::ThumbnailTable,				Summary.ThumbnailTableOffset,			HeaderInformation.ThumbnailTableSize,				false	},
 		{ EPatchedSection::AssetRegistryData,			Summary.AssetRegistryDataOffset,		AssetRegistryData.SectionSize,						true	},
 	};
@@ -998,7 +2125,7 @@ FAssetHeaderPatcher::EResult FAssetHeaderPatcherInner::PatchHeader_WriteDestinat
 	const int32 SourceTotalHeaderSize = Summary.TotalHeaderSize;
 
 	// Ensure the sections are in the expected order.
-	for (int SectionIdx = 1; SectionIdx < UE_ARRAY_COUNT(SourceSections); ++SectionIdx)
+	for (int32 SectionIdx = 1; SectionIdx < UE_ARRAY_COUNT(SourceSections); ++SectionIdx)
 	{
 		const FSectionData& SourceSection = SourceSections[SectionIdx];
 		const FSectionData& PrevSection = SourceSections[SectionIdx - 1];
@@ -1049,7 +2176,7 @@ FAssetHeaderPatcher::EResult FAssetHeaderPatcherInner::PatchHeader_WriteDestinat
 
 	int64 LastSectionEndedAt = 0;
 
-	for (int SectionIdx = 0; SectionIdx < UE_ARRAY_COUNT(SourceSections); ++SectionIdx)
+	for (int32 SectionIdx = 0; SectionIdx < UE_ARRAY_COUNT(SourceSections); ++SectionIdx)
 	{
 		const FSectionData& SourceSection = SourceSections[SectionIdx];
 
@@ -1144,6 +2271,24 @@ FAssetHeaderPatcher::EResult FAssetHeaderPatcherInner::PatchHeader_WriteDestinat
 			break;
 		}
 
+		case EPatchedSection::SearchableNamesMap:
+		{
+			const int64 TableStartOffset = Writer.Tell();
+			
+			FLinkerTables LinkerTables;
+			LinkerTables.SearchableNamesMap = SearchableNamesMap;
+			LinkerTables.SerializeSearchableNamesMap(Writer);
+
+			checkf(!Writer.IsCriticalError(), TEXT("Issue writing %s"), *Writer.GetErrorMessage());
+
+			const int64 TableSize = Writer.Tell() - TableStartOffset;
+			const int64 Delta = TableSize - SourceSection.Size;
+			checkf(Delta == 0, TEXT("Delta should be Zero. is %d"), (int)Delta);
+			check(Summary.SearchableNamesOffset == TableStartOffset);
+
+			break;
+		}
+
 		case EPatchedSection::ImportTable:
 		{
 			const int64 ImportTableStartOffset = Writer.Tell();
@@ -1204,23 +2349,26 @@ FAssetHeaderPatcher::EResult FAssetHeaderPatcherInner::PatchHeader_WriteDestinat
 
 		case EPatchedSection::ThumbnailTable:
 		{
-			// Luckily the strings in the thumbnail table don't include the package, so only the offsets need to be patched
 			const int64 ThumbnailTableStartOffset = Writer.Tell();
 			const int64 ThumbnailTableDeltaOffset = ThumbnailTableStartOffset - SourceSection.Offset;
 			int32 ThumbnailCount = ThumbnailTable.Num();
 			Writer << ThumbnailCount;
+			int32 AccumulatedDelta = 0;
 			for (FThumbnailEntry& Entry : ThumbnailTable)
 			{
+				AccumulatedDelta += Entry.Delta;
 				Writer << Entry.ObjectShortClassName;
 				Writer << Entry.ObjectPathWithoutPackageName;
-				Entry.FileOffset += (int32)ThumbnailTableDeltaOffset; // Thumbnail payloads immediately follow the table, so we can just apply the section offset delta here
+				Entry.FileOffset += (int32)ThumbnailTableDeltaOffset + AccumulatedDelta;
 				Writer << Entry.FileOffset;
 			}
 			checkf(!Writer.IsCriticalError(), TEXT("Issue writing %s"), *Writer.GetErrorMessage());
 
 			const int64 ThumbnailTableSize = Writer.Tell() - ThumbnailTableStartOffset;
+			const int64 Delta = ThumbnailTableSize - SourceSection.Size;
+			PatchSummaryOffsets(Summary, ThumbnailTableStartOffset, Delta);
+			Summary.TotalHeaderSize += (int32)Delta;
 			checkf(ThumbnailTableStartOffset == Summary.ThumbnailTableOffset, TEXT("%zd == %zd"), ThumbnailTableStartOffset, Summary.ThumbnailTableOffset);
-			check(ThumbnailTableSize == SourceSection.Size); // We only patch thumbnail table offsets, we should not be patching size
 
 			break;
 		}
@@ -1330,3 +2478,855 @@ FAssetHeaderPatcher::EResult FAssetHeaderPatcherInner::PatchHeader_WriteDestinat
 
 	return EResult::Success;
 }
+
+void FAssetHeaderPatcherInner::DumpState(FStringView OutputDirectory)
+{
+	TStringBuilder<1024> Builder;
+	auto GetDebugFNameString = [this](FName Name)
+		{
+			int32* Index = NameToIndexMap.Find(Name.GetDisplayIndex());
+			if (Index)
+			{
+				return FString::Printf(TEXT("%s (nametable index: %d, fname:{'%s', %d})"), *Name.ToString(), *Index, *Name.GetPlainNameString(), Name.GetNumber());
+			}
+			else
+			{
+				return FString(TEXT("None (nametable index: -1, fname {'None', 0})"));
+			}
+		};
+
+	Builder.Append(TEXT("{\n"));
+
+	Builder.Append(TEXT("\t\"Summary\":{ "));
+	{
+		Builder.Append(TEXT("\n\t\t\"PackageName\": \""));
+		Builder.Append(Summary.PackageName);
+		Builder.Append(TEXT("\""));
+	}
+	Builder.Append(TEXT("\n\t},\n"));
+
+	Builder.Append(TEXT("\t\"NameTable\":[ "));
+	for (const FName& Name : NameTable)
+	{
+		Builder.Append(TEXT("\n\t\t\""));
+		Builder.Append(GetDebugFNameString(Name));
+		Builder.Append(TEXT("\","));
+	}
+	Builder.RemoveSuffix(1);
+	Builder.Append(TEXT("\n\t],\n"));
+
+	Builder.Append(TEXT("\t\"ExportTable\":[ "));
+	for (const auto& Export : ExportTable)
+	{
+		Builder.Append(TEXT("\n\t\t{\n"));
+
+		Builder.Append(TEXT("\t\t\t\"ObjectName\": \""));
+		Builder.Append(GetDebugFNameString(Export.ObjectName));
+		Builder.Append(TEXT("\""));
+
+#if WITH_EDITORONLY_DATA
+		Builder.Append(TEXT(",\n"));
+		Builder.Append(TEXT("\t\t\t\"OldClassName\": \""));
+		Builder.Append(GetDebugFNameString(Export.OldClassName));
+		Builder.Append(TEXT("\""));
+#endif
+		Builder.Append(TEXT("\n\t\t},"));
+	}
+	Builder.RemoveSuffix(1);
+	Builder.Append(TEXT("\n\t],\n"));
+
+	Builder.Append(TEXT("\t\"ImportTable\":[ "));
+	for (const auto& Import : ImportTable)
+	{
+		Builder.Append(TEXT("\n\t\t{\n"));
+
+		Builder.Append(TEXT("\t\t\t\"ObjectName\": \""));
+		Builder.Append(GetDebugFNameString(Import.ObjectName));
+		Builder.Append(TEXT("\",\n"));
+
+#if WITH_EDITORONLY_DATA
+		Builder.Append(TEXT("\t\t\t\"OldClassName\": \""));
+		Builder.Append(GetDebugFNameString(Import.OldClassName));
+		Builder.Append(TEXT("\",\n"));
+#endif
+
+		Builder.Append(TEXT("\t\t\t\"ClassPackage\": \""));
+		Builder.Append(GetDebugFNameString(Import.ClassPackage));
+		Builder.Append(TEXT("\",\n"));
+
+		Builder.Append(TEXT("\t\t\t\"ClassName\": \""));
+		Builder.Append(GetDebugFNameString(Import.ClassName));
+		Builder.Append(TEXT("\""));
+
+#if WITH_EDITORONLY_DATA
+		Builder.Append(TEXT(",\n\t\t\t\"PackageName\": \""));
+		Builder.Append(GetDebugFNameString(Import.PackageName));
+		Builder.Append(TEXT("\""));
+#endif
+		Builder.Append(TEXT("\n\t\t},"));
+
+	}
+	Builder.RemoveSuffix(1);
+	Builder.Append(TEXT("\n\t],\n"));
+
+	Builder.Append(TEXT("\t\"SoftObjectPathTable\":[ "));
+	for (const FSoftObjectPath& SoftObjectPath : SoftObjectPathTable)
+	{
+		Builder.Append(TEXT("\n\t\t{\n"));
+
+		FTopLevelAssetPath TLAP = SoftObjectPath.GetAssetPath();
+		FString Subpath = SoftObjectPath.GetSubPathString();
+
+		Builder.Append(TEXT("\t\t\t\"AssetPath\": {\n\""));
+		Builder.Append(TEXT("\t\t\t\t\"PackageName\": \""));
+		Builder.Append(GetDebugFNameString(TLAP.GetPackageName()));
+		Builder.Append(TEXT("\",\n"));
+		Builder.Append(TEXT("\t\t\t\t\"AssetName\": \""));
+		Builder.Append(GetDebugFNameString(TLAP.GetAssetName()));
+		Builder.Append(TEXT("\"\n"));
+		Builder.Append(TEXT("\t\t\t},\n"));
+
+		Builder.Append(TEXT("\t\t\t\"Subpath (string)\": \""));
+		Builder.Append(Subpath);
+		Builder.Append(TEXT("\""));
+
+		Builder.Append(TEXT("\n\t\t},"));
+
+	}
+	Builder.RemoveSuffix(1);
+	Builder.Append(TEXT("\n\t],\n"));
+
+	Builder.Append(TEXT("\t\"SoftPackageReferencesTable\":[ "));
+	for (const FName SoftPackageRef : SoftPackageReferencesTable)
+	{
+		Builder.Append(TEXT("\n\t\t\""));
+		Builder.Append(GetDebugFNameString(SoftPackageRef));
+		Builder.Append(TEXT("\","));
+	}
+	Builder.RemoveSuffix(1);
+	Builder.Append(TEXT("\n\t],\n"));
+
+	Builder.Append(TEXT("\t\"GatherableTextDataTable\":[ "));
+	for (const FGatherableTextData& GatherableTextData : GatherableTextDataTable)
+	{
+		Builder.Append(TEXT("\n\t\t{\n"));
+		Builder.Append(TEXT("\t\t\t\"SourceSiteContexts.SiteDescription (string)\": ["));
+		for (auto& SiteContext : GatherableTextData.SourceSiteContexts)
+		{
+			Builder.Append(TEXT("\n\t\t\t\t\""));
+			Builder.Append(SiteContext.SiteDescription);
+			Builder.Append(TEXT("\","));
+		}
+		Builder.RemoveSuffix(1);
+		Builder.Append(TEXT("\n\t\t\t]"));
+		Builder.Append(TEXT("\n\t\t},"));
+	}
+	Builder.RemoveSuffix(1);
+	Builder.Append(TEXT("\n\t],\n"));
+
+	Builder.Append(TEXT("\t\"ThumbnailTable\":[ "));
+	for (const FThumbnailEntry& ThumbnailEntry : ThumbnailTable)
+	{
+		Builder.Append(TEXT("\n\t\t{\n"));
+
+		Builder.Append(TEXT("\t\t\t\"ObjectPathWithoutPackageName (string)\": \""));
+		Builder.Append(ThumbnailEntry.ObjectPathWithoutPackageName);
+		Builder.Append(TEXT("\",\n"));
+
+		Builder.Append(TEXT("\t\t\t\"ObjectShortClassName (string)\": \""));
+		Builder.Append(ThumbnailEntry.ObjectShortClassName);
+		Builder.Append(TEXT("\""));
+
+		Builder.Append(TEXT("\n\t\t},"));
+	}
+	Builder.RemoveSuffix(1);
+	Builder.Append(TEXT("\n\t],\n"));
+
+	Builder.Append(TEXT("\t\"AssetRegistryData\":[ "));
+	for (const FAssetRegistryObjectData& ObjData : AssetRegistryData.ObjectData)
+	{
+		Builder.Append(TEXT("\n\t\t{\n"));
+
+		Builder.Append(TEXT("\t\t\t\"ObjectData\": {\n"));
+		Builder.Append(TEXT("\t\t\t\t\"ObjectPath (string)\": \""));
+		Builder.Append(ObjData.ObjectData.ObjectPath);
+		Builder.Append(TEXT("\",\n"));
+		Builder.Append(TEXT("\t\t\t\t\"ObjectClassName (string)\": \""));
+		Builder.Append(ObjData.ObjectData.ObjectClassName);
+		Builder.Append(TEXT("\"\n"));
+		Builder.Append(TEXT("\t\t\t},\n"));
+
+		Builder.Append(TEXT("\t\t\t\"TagData\": [\n"));
+		for (const auto& TagData : ObjData.TagData)
+		{
+			FString Value = TagData.Value;
+			bool bNeedDecode = TagData.Key == FWorldPartitionActorDescUtils::ActorMetaDataTagName();
+			if (bNeedDecode)
+			{
+				const FString LongPackageName(SrcAsset);
+				const FString ObjectPath(ObjData.ObjectData.ObjectPath);
+				const FTopLevelAssetPath AssetClassPathName(ObjData.ObjectData.ObjectClassName);
+				const FAssetDataTagMap Tags(MakeTagMap(ObjData.TagData));
+				const FAssetData AssetData(LongPackageName, ObjectPath, AssetClassPathName, Tags);
+
+				struct FWorldPartitionAssetDataPrinter : FWorldPartitionAssetDataPatcher
+				{
+					FWorldPartitionAssetDataPrinter(int32 InIndentDepth)
+						: IndentDepth(InIndentDepth) 
+					{
+					}
+
+					virtual bool DoPatch(FString& InOutString) override
+					{
+						Builder.Append(TEXT("\n"));
+						Indent();
+						Builder.Append(TEXT("string=\""));
+						Builder.Append(InOutString);
+						Builder.Append(TEXT("\""));
+						return false;
+					}
+					virtual bool DoPatch(FName& InOutName) override
+					{
+						Builder.Append(TEXT("\n"));
+						Indent();
+						Builder.Append(TEXT("FName=\""));
+						Builder.Append(InOutName.ToString());
+						Builder.Append(TEXT("\""));
+						return false;
+					}
+					virtual bool DoPatch(FSoftObjectPath& InOutSoft) override
+					{
+						Builder.Append(TEXT("\n"));
+						Indent();
+						Builder.Append(TEXT("FSoftObjectPath="));
+						FTopLevelAssetPath TLAP = InOutSoft.GetAssetPath();
+						Builder.Append(TEXT("{{PackageName=\""));
+						Builder.Append(TLAP.GetPackageName().ToString());
+						Builder.Append(TEXT("\", AssetName=\""));
+						Builder.Append(TLAP.GetAssetName().ToString());
+						Builder.Append(TEXT("\"}, SubPath (string)=\""));
+						Builder.Append(InOutSoft.GetSubPathString());
+						Builder.Append(TEXT("\"}"));
+						return false;
+					}
+					virtual bool DoPatch(FTopLevelAssetPath& InOutPath) override
+					{
+						Builder.Append(TEXT("\n"));
+						Indent();
+						Builder.Append(TEXT("FTopLevelAssetPath="));
+						Builder.Append(TEXT("{PackageName=\""));
+						Builder.Append(InOutPath.GetPackageName().ToString());
+						Builder.Append(TEXT("\", AssetName=\""));
+						Builder.Append(InOutPath.GetAssetName().ToString());
+						Builder.Append(TEXT("\"}"));
+						return false;
+					}
+					void Indent()
+					{
+						for (int32 i = 0; i < IndentDepth; ++i)
+						{
+							Builder.Append(TEXT("\t"));
+						}
+					}
+					const TCHAR* ToString() const
+					{
+						return Builder.ToString();
+					}
+					int32 IndentDepth;
+					TStringBuilder<1024> Builder;
+				};
+
+				FString PatchedAssetData;
+				FWorldPartitionAssetDataPrinter Patcher(5);
+				FWorldPartitionActorDescUtils::GetPatchedAssetDataFromAssetData(AssetData, PatchedAssetData, &Patcher);
+				Value = Patcher.ToString();
+			}
+
+			Builder.Append(TEXT("\n\t\t\t\t{\n"));
+
+			Builder.Append(TEXT("\t\t\t\t\t\"Key (string)\": \""));
+			Builder.Append(TagData.Key);
+			Builder.Append(TEXT("\",\n"));
+			Builder.Append(TEXT("\t\t\t\t\t\"Value"));
+			if (bNeedDecode)
+			{
+				Builder.Append(TEXT(" (decoded string)"));
+			}
+			else
+			{
+				Builder.Append(TEXT("(string)"));
+			}
+			Builder.Append(TEXT("\": \""));
+			Builder.Append(Value);
+			Builder.Append(TEXT("\"\n"));
+
+			Builder.Append(TEXT("\t\t\t\t},"));
+		}
+		Builder.RemoveSuffix(1);
+		Builder.Append(TEXT("\n\t\t\t]\n"));
+
+		Builder.Append(TEXT("\n\t\t},"));
+	}
+	Builder.RemoveSuffix(1);
+	Builder.Append(TEXT("\n\t]\n"));
+
+	Builder.Append(TEXT("}"));
+
+
+	// Write to disk
+	TStringBuilder<256> OutPath;
+	OutPath.Append(OutputDirectory);
+	FString SubPath = SrcAsset;
+	FPaths::CollapseRelativeDirectories(SubPath);
+	if (SubPath.StartsWith(TEXT("../")))
+	{
+		int32 Pos = SubPath.Find(TEXT("../"), ESearchCase::CaseSensitive, ESearchDir::FromEnd);
+		if (Pos >= 0)
+		{
+			SubPath.RightChopInline(Pos + 3);
+		}
+	}
+	else if (SubPath.Len() > 2 && SubPath[1] == TEXT(':'))
+	{
+		SubPath.RightChopInline(2); // Drop the drive
+	}
+	OutPath = FPaths::Combine(OutPath, SubPath);
+	OutPath.Append(TEXT(".txt"));
+	FFileHelper::SaveStringToFile(Builder.ToString(), *OutPath);
+}
+
+///////////////////////////////////////////////////////////////////////////
+
+#if WITH_TESTS
+
+#include "Tests/TestHarnessAdapter.h"
+
+TEST_CASE("AssetHeaderPatcher", "[AssetHeaderPatcher][EngineFilter]")
+{
+// Useful when iterating so you halt if something fails while a debugger is attached
+#if 0
+#undef CHECK
+#undef CHECK_EQUALS
+#undef CHECK_NOT_EQUALS
+#define CHECK(...) if (!(__VA_ARGS__)) { UE_DEBUG_BREAK(); FAutomationTestFramework::Get().GetCurrentTest()->AddError(TEXT("Condition failed")); }
+#define CHECK_EQUALS(What, X, Y) if(!FAutomationTestFramework::Get().GetCurrentTest()->TestEqual(What, X, Y)) { UE_DEBUG_BREAK(); };
+#define CHECK_NOT_EQUALS(What, X, Y) if(!FAutomationTestFramework::Get().GetCurrentTest()->TestNotEqual(What, X, Y)) { UE_DEBUG_BREAK(); };
+#endif
+
+	struct FTestPatcherContext : FAssetHeaderPatcher::FContext
+	{
+		FTestPatcherContext(TMap<FString, FString> PackageRenameMap, bool bGatherDependentPackages = true) : FAssetHeaderPatcher::FContext(PackageRenameMap, bGatherDependentPackages) {}
+		const TMap<FString, FString>& GetStringReplacements()
+		{
+			return StringReplacements;
+		}
+
+		void GenerateRemappings()
+		{
+			GenerateAdditionalRemappings();
+		}
+
+		const TArray<FCoreRedirect>& GetRedirects()
+		{
+			return Redirects;
+		}
+
+		const TArray<FString>& GetVerseMountPoints()
+		{
+			return VerseMountPoints;
+		}
+	};
+
+	// To avoid having to deal with serialization, we mock some data and inject it directly
+	// into the patcher as if done via serialization
+	const FString DummySrcDstAsset			= TEXT("/SrcMount/SomePath/SrcPackage");
+	const TCHAR*  SrcPackagePath			= TEXT("/SrcMount/SomePath/SrcPackage");
+	const TCHAR*  DstPackagePath			= TEXT("/DstMount/SomePath/DstPackage");
+	const TCHAR*  SrcPackageObjectPath		= TEXT("/SrcMount/SomePath/SrcPackage.SrcPackage");
+	const TCHAR*  DstPackageObjectPath		= TEXT("/DstMount/SomePath/DstPackage.DstPackage");
+	const TCHAR*  SrcMountName				= TEXT("/SrcSpecialMount/");
+	const TCHAR*  DstMountName				= TEXT("/DstSpecialMount/");
+	const FName   SrcPackagePathFName(SrcPackagePath);
+	const FName   DstPackagePathFName(DstPackagePath);
+	const FName   SrcAssetFName(TEXT("SrcPackage"));
+	const FName   DstAssetFName(TEXT("DstPackage"));
+	const FName   SrcExportObjectFName = SrcAssetFName;
+	const FName   DstExportObjectFName = DstAssetFName;
+	const FName   DummyImportPackagePathFName(TEXT("/DummyMount/DummyPackage"));
+
+
+	TMap<FString, FString> MountPointReplacementMap =
+	{
+		{ SrcMountName, DstMountName },
+	};
+
+	TMap<FString, FString> PackageRenameMap =
+	{
+		{ SrcPackagePath, DstPackagePath },
+	};
+
+	FTestPatcherContext Context(PackageRenameMap, false /*bGatherDependentPackages*/);
+	const TMap<FString, FString>& StringReplacements = Context.GetStringReplacements();
+	CHECK(StringReplacements.Num() > PackageRenameMap.Num()); // Ensure we generated more mappings off of the PackageRenameMap
+	CHECK(FCoreRedirects::AddRedirectList(Context.GetRedirects(), TEXT("Asset Header Patcher Tests")));
+
+	FAssetHeaderPatcherInner Patcher(DummySrcDstAsset, DummySrcDstAsset, StringReplacements, MountPointReplacementMap);
+
+	int32 OriginalNameTableCount = 0;
+	auto ResetPatcher = [&Patcher, &OriginalNameTableCount, SrcPackagePathFName, SrcAssetFName, SrcExportObjectFName, DummyImportPackagePathFName]()
+		{
+			// Reset NameTable
+			Patcher.NameTable.Empty();
+			Patcher.NameToIndexMap.Empty();
+			Patcher.RenameMap.Empty();
+			Patcher.AddedNames.Empty();
+			Patcher.ExportTable.Empty();
+
+			// Repopulate it with our test data normally set through deserialization
+			// NameTable
+			Patcher.NameToIndexMap.Add(SrcPackagePathFName.GetDisplayIndex(), Patcher.NameTable.Num());
+			Patcher.NameTable.Add(SrcPackagePathFName);
+			Patcher.NameToIndexMap.Add(SrcAssetFName.GetDisplayIndex(), Patcher.NameTable.Num());
+			Patcher.NameTable.Add(SrcAssetFName);
+			Patcher.NameToIndexMap.Add(FName(NAME_None).GetDisplayIndex(), Patcher.NameTable.Num());
+			Patcher.NameTable.Add(FName(NAME_None));
+			Patcher.NameToIndexMap.Add(FName(DummyImportPackagePathFName).GetDisplayIndex(), Patcher.NameTable.Num());
+			Patcher.NameTable.Add(FName(DummyImportPackagePathFName));
+
+			FObjectImport DummyImport;
+			DummyImport.ObjectName = Patcher.NameTable[Patcher.NameToIndexMap[DummyImportPackagePathFName.GetDisplayIndex()]];
+			DummyImport.OldClassName = DummyImportPackagePathFName; // Set to something other than NAME_None
+			DummyImport.OuterIndex = FPackageIndex();
+			Patcher.ImportTable.Add(DummyImport);
+
+			// Export Table
+			FObjectExport SrcPackageExport;
+			SrcPackageExport.ObjectName = Patcher.NameTable[Patcher.NameToIndexMap[SrcExportObjectFName.GetDisplayIndex()]];
+			SrcPackageExport.OldClassName = SrcExportObjectFName; // Set to something other than NAME_None
+			SrcPackageExport.OuterIndex = FPackageIndex(); // This package is the outer
+			Patcher.ExportTable.Add(SrcPackageExport);
+
+			FObjectExport DummyExport;
+			DummyExport.ObjectName = Patcher.NameTable[Patcher.NameToIndexMap[SrcExportObjectFName.GetDisplayIndex()]]; // Same name as SrcPackageExport
+			DummyExport.OldClassName = SrcExportObjectFName; // Set to something other than NAME_None
+			DummyExport.OuterIndex = FPackageIndex::FromImport(0); // DummyImport is our outer
+			Patcher.ExportTable.Add(DummyExport);
+
+			// Summary
+			Patcher.Summary.NameCount = Patcher.NameTable.Num();
+			Patcher.OriginalPackagePath = SrcPackagePathFName;
+			OriginalNameTableCount = Patcher.NameTable.Num();
+		};
+
+	SECTION("FContext Additional Remappings")
+	{
+		{
+			FString Actual = LR"(/SrcMount/SomePath/SrcPackage)";
+			const FString Expected = LR"(/DstMount/SomePath/DstPackage)";
+			CHECK(Patcher.DoPatch(Actual));
+			CHECK_EQUALS(TEXT("Patch string with direct match"), Actual, Expected);
+		}
+
+		{
+			FString Actual = LR"(/SrcMount/SomePath/SrcPackage.SrcPackage)";
+			const FString Expected = LR"(/DstMount/SomePath/DstPackage.DstPackage)";
+			CHECK(Patcher.DoPatch(Actual));
+			CHECK_EQUALS(TEXT("Generated Top-Level Asset mapping"), Actual, Expected);
+		}
+
+		{
+			FString Actual = LR"(/SrcMount/SomePath/SrcPackage.SrcPackage_C)";
+			const FString Expected = LR"(/DstMount/SomePath/DstPackage.DstPackage_C)";
+			CHECK(Patcher.DoPatch(Actual));
+			CHECK_EQUALS(TEXT("Generated Blueprint Generated Class mapping"), Actual, Expected);
+		}
+
+		{
+			FString Actual = LR"(/SrcMount/SomePath/SrcPackage.Default__SrcPackage_C)";
+			const FString Expected = LR"(/DstMount/SomePath/DstPackage.Default__DstPackage_C)";
+			CHECK(Patcher.DoPatch(Actual));
+			CHECK_EQUALS(TEXT("Generated Blueprint Generated Class Default Object mapping"), Actual, Expected);
+		}
+
+		{
+			FString Actual = LR"(/SrcMount/SomePath/SrcPackage.SrcPackageEditorOnlyData)";
+			const FString Expected = LR"(/DstMount/SomePath/DstPackage.DstPackageEditorOnlyData)";
+			CHECK(Patcher.DoPatch(Actual));
+			CHECK_EQUALS(TEXT("Generated MaterialFunctionInterface Editor Only Data mapping"), Actual, Expected);
+		}
+
+		SECTION("Verse Mountpoints")
+		{
+			for (const FString& VerseMount : Context.GetVerseMountPoints())
+			{
+				// We only generate verse paths for objects, so this package path will not have a mapping
+				{
+					FString Actual = FString::Printf(LR"(/%s/SrcMount/SomePath/SrcPackage)", *VerseMount);
+					const FString Expected = FString::Printf(LR"(/%s/DstMount/SomePath/DstPackage)", *VerseMount);
+					CHECK(!Patcher.DoPatch(Actual));
+					CHECK_NOT_EQUALS(TEXT("Patch string with direct match"), Actual, Expected);
+				}
+
+				{
+					FString Actual = FString::Printf(LR"(/%s/SrcMount/SomePath/SrcPackage/SrcPackage)", *VerseMount);
+					const FString Expected = FString::Printf(LR"(/%s/DstMount/SomePath/DstPackage/DstPackage)", *VerseMount);
+					CHECK(Patcher.DoPatch(Actual));
+					CHECK_EQUALS(TEXT("Patch string with direct match"), Actual, Expected);
+				}
+
+				{
+					FString Actual = FString::Printf(LR"(/%s/SrcMount/SomePath/SrcPackage/SrcPackage)", *VerseMount);
+					const FString Expected = FString::Printf(LR"(/%s/DstMount/SomePath/DstPackage/DstPackage)", *VerseMount);
+					CHECK(Patcher.DoPatch(Actual));
+					CHECK_EQUALS(TEXT("Generated Top-Level Asset mapping"), Actual, Expected);
+				}
+
+				{
+					FString Actual = FString::Printf(LR"(/%s/SrcMount/SomePath/SrcPackage/SrcPackage_C)", *VerseMount);
+					const FString Expected = FString::Printf(LR"(/%s/DstMount/SomePath/DstPackage/DstPackage_C)", *VerseMount);
+					CHECK(Patcher.DoPatch(Actual));
+					CHECK_EQUALS(TEXT("Generated Blueprint Generated Class mapping"), Actual, Expected);
+				}
+
+				{
+					FString Actual = FString::Printf(LR"(/%s/SrcMount/SomePath/SrcPackage/Default__SrcPackage_C)", *VerseMount);
+					const FString Expected = FString::Printf(LR"(/%s/DstMount/SomePath/DstPackage/Default__DstPackage_C)", *VerseMount);
+					CHECK(Patcher.DoPatch(Actual));
+					CHECK_EQUALS(TEXT("Generated Blueprint Generated Class Default Object mapping"), Actual, Expected);
+				}
+
+				{
+					FString Actual = FString::Printf(LR"(/%s/SrcMount/SomePath/SrcPackage/SrcPackageEditorOnlyData)", *VerseMount);
+					const FString Expected = FString::Printf(LR"(/%s/DstMount/SomePath/DstPackage/DstPackageEditorOnlyData)", *VerseMount);
+					CHECK(Patcher.DoPatch(Actual));
+					CHECK_EQUALS(TEXT("Generated MaterialFunctionInterface Editor Only Data mapping"), Actual, Expected);
+				}
+			}
+		}
+	}
+
+	SECTION("DoPatch(FString)")
+	{
+		SECTION("Direct match")
+		{
+			{
+				FString Actual = LR"(/SrcMount/SomePath/SrcPackage)";
+				const FString Expected = LR"(/DstMount/SomePath/DstPackage)";
+				CHECK(Patcher.DoPatch(Actual));
+				CHECK_EQUALS(TEXT("Patch string with direct match"), Actual, Expected);
+			}
+
+			{
+				FString Actual = LR"(/SrcMount/SomePath/SrcPackage2)";
+				const FString Expected = Actual; // Must be a copy
+				CHECK(!Patcher.DoPatch(Actual));
+				CHECK_EQUALS(TEXT("Patch string with no direct match"), Actual, Expected);
+			}
+		}
+
+		SECTION("Sub-Object Paths")
+		{
+			{
+				FString Actual = LR"(/SrcMount/SomePath/SrcPackage.SrcPackage:AnOuter.To.A.SubObject)";
+				const FString Expected = LR"(/DstMount/SomePath/DstPackage.DstPackage:AnOuter.To.A.SubObject)";
+				CHECK(Patcher.DoPatch(Actual));
+				CHECK_EQUALS(TEXT("Patch sub-object path"), Actual, Expected);
+			}
+
+			// Worth adding support for in the future, but at the moment we cannot patch various parts of unquoted 
+			// sub-object paths (that are specifically strings in the header, FNames are fine). In this case we 
+			// can't patch the package path because the top-level asset (UnmappedObject) has no mapping for patching
+			{
+				FString Actual = LR"(/SrcMount/SomePath/SrcPackage.UnmappedObject:AnOuter.To.A.SubObject)";
+				const FString Expected = LR"(/DstMount/SomePath/DstPackage.UnmappedObject:AnOuter.To.A.SubObject)";
+				CHECK(!Patcher.DoPatch(Actual));
+				CHECK_NOT_EQUALS(TEXT("Can't patch sub-object paths, for "), Actual, Expected);
+			}
+		}
+
+		SECTION("Quoted match")
+		{
+			SECTION("Single Quote")
+			{
+				{
+					FString Actual = LR"('/SrcMount/SomePath/SrcPackage')";
+					const FString Expected = LR"('/DstMount/SomePath/DstPackage')";
+					CHECK(Patcher.DoPatch(Actual));
+					CHECK_EQUALS(TEXT("Patch package path with quotes"), Actual, Expected);
+				}
+
+				{
+					FString Actual = LR"('/SrcMount/SomePath/SrcPackage.SrcPackage')";
+					const FString Expected = LR"('/DstMount/SomePath/DstPackage.DstPackage')";
+					CHECK(Patcher.DoPatch(Actual));
+					CHECK_EQUALS(TEXT("Patch object path with quotes"), Actual, Expected);
+				}
+
+				{
+					FString Actual = LR"('/SrcMount/SomePath/SrcPackage.SrcPackage_C')";
+					const FString Expected = LR"('/DstMount/SomePath/DstPackage.DstPackage_C')";
+					CHECK(Patcher.DoPatch(Actual));
+					CHECK_EQUALS(TEXT("Patch blueprint generated class with quotes"), Actual, Expected);
+				}
+
+				{
+					FString Actual = LR"('/SrcMount/SomePath/SrcPackage.Default__SrcPackage_C')";
+					const FString Expected = LR"('/DstMount/SomePath/DstPackage.Default__DstPackage_C')";
+					CHECK(Patcher.DoPatch(Actual));
+					CHECK_EQUALS(TEXT("Patch default blueprint generated class object path with quotes"), Actual, Expected);
+				}
+			}
+
+			SECTION("Double Quote")
+			{
+				{
+					FString Actual = LR"("/SrcMount/SomePath/SrcPackage")";
+					const FString Expected = LR"("/DstMount/SomePath/DstPackage")";
+					CHECK(Patcher.DoPatch(Actual));
+					CHECK_EQUALS(TEXT("Patch package path with quotes"), Actual, Expected);
+				}
+
+				{
+					FString Actual = LR"("/SrcMount/SomePath/SrcPackage.SrcPackage")";
+					const FString Expected = LR"("/DstMount/SomePath/DstPackage.DstPackage")";
+					CHECK(Patcher.DoPatch(Actual));
+					CHECK_EQUALS(TEXT("Patch object path with quotes"), Actual, Expected);
+				}
+
+				{
+					FString Actual = LR"("/SrcMount/SomePath/SrcPackage.SrcPackage_C")";
+					const FString Expected = LR"("/DstMount/SomePath/DstPackage.DstPackage_C")";
+					CHECK(Patcher.DoPatch(Actual));
+					CHECK_EQUALS(TEXT("Patch blueprint generated class with quotes"), Actual, Expected);
+				}
+
+				{
+					FString Actual = LR"("/SrcMount/SomePath/SrcPackage.Default__SrcPackage_C")";
+					const FString Expected = LR"("/DstMount/SomePath/DstPackage.Default__DstPackage_C")";
+					CHECK(Patcher.DoPatch(Actual));
+					CHECK_EQUALS(TEXT("Patch default blueprint generated class object path with quotes"), Actual, Expected);
+				}
+			}
+
+			SECTION("Substring match")
+			{
+				{
+					FString Actual =
+						LR"(((ReferenceNodePath="/SrcMount/SomePath/SrcPackage.SrcPackage:RigVMModel.Setup Arm",)"
+						LR"(((Package="/SrcMount/SomePath/SrcPackage",)"
+						LR"(HostObject="/SrcMount/SomePath/SrcPackage.SrcPackage_C")))";
+					FString Expected =
+						LR"(((ReferenceNodePath="/DstMount/SomePath/DstPackage.DstPackage:RigVMModel.Setup Arm",)"
+						LR"(((Package="/DstMount/SomePath/DstPackage",)"
+						LR"(HostObject="/DstMount/SomePath/DstPackage.DstPackage_C")))";
+					CHECK(Patcher.DoPatch(Actual));
+					CHECK_EQUALS(TEXT("Patch substring with quoted package, object and sub-object paths"), Actual, Expected);
+				}
+			}
+		}
+
+		SECTION("Mountpoint match")
+		{
+			/*
+				We currently don't support mount point replacement _for strings_ that don't
+				provide some kind of delimiter for us to scan for. As such package paths
+				and top-level asset paths are not supported unless they are quoted. Sub-object
+				paths are supported.
+			*/
+			/*
+			{
+				FString Actual = LR"(/SrcSpecialMount/SomePath/SomePackage)";
+				const FString Expected = LR"(/DstSpecialMount/SomePath/SomePackage)";
+				CHECK(Patcher.DoPatch(Actual));
+				CHECK_EQUALS(TEXT("Patch package path replaces only mount"), Actual, Expected);
+			}
+
+			{
+				FString Actual = LR"(/SrcSpecialMount/SomePath/SomePackage.SomePackage)";
+				const FString Expected = LR"(/DstSpecialMount/SomePath/SomePackage.SomePackage)";
+				CHECK(Patcher.DoPatch(Actual));
+				CHECK_EQUALS(TEXT("Patch package object path replaces only mount"), Actual, Expected);
+			}
+			*/
+
+			{
+				FString Actual = LR"(/SrcSpecialMount/SomePath/SomePackage.TopLevel:SubObject1.SubObject2)";
+				const FString Expected = LR"(/DstSpecialMount/SomePath/SomePackage.TopLevel:SubObject1.SubObject2)";
+				CHECK(Patcher.DoPatch(Actual));
+				CHECK_EQUALS(TEXT("Patch package sub-object path replaces only mount"), Actual, Expected);
+			}
+
+			{
+				FString Actual = LR"("/SrcSpecialMount/SomePath/SomePackage")";
+				const FString Expected = LR"("/DstSpecialMount/SomePath/SomePackage")";
+				CHECK(Patcher.DoPatch(Actual));
+				CHECK_EQUALS(TEXT("Patch double quoted path replaces only mount"), Actual, Expected);
+			}
+
+			{
+				FString Actual = LR"('/SrcSpecialMount/SomePath/SomePackage')";
+				const FString Expected = LR"('/DstSpecialMount/SomePath/SomePackage')";
+				CHECK(Patcher.DoPatch(Actual));
+				CHECK_EQUALS(TEXT("Patch single quoted path replaces only mount"), Actual, Expected);
+			}
+
+			{
+				FString Actual = LR"(SomePrefix="/SrcSpecialMount/SomePath/SomePackage")";
+				const FString Expected = LR"(SomePrefix="/DstSpecialMount/SomePath/SomePackage")";
+				CHECK(Patcher.DoPatch(Actual));
+				CHECK_EQUALS(TEXT("Substring patch replaces only mount when double quoted"), Actual, Expected);
+			}
+
+			{
+				FString Actual = LR"(SomePrefix='/SrcSpecialMount/SomePath/SomePackage')";
+				const FString Expected = LR"(SomePrefix='/DstSpecialMount/SomePath/SomePackage')";
+				CHECK(Patcher.DoPatch(Actual));
+				CHECK_EQUALS(TEXT("Substring patch replaces only mount when single quoted"), Actual, Expected);
+			}
+		}
+	}
+
+	SECTION("DoPatch(FSoftObjectPath")
+	{
+		{
+			ResetPatcher();
+
+			FSoftObjectPath Actual(TEXT("/SrcMount/SomePath/SrcPackage.SrcPackage"));
+			FSoftObjectPath Expected(TEXT("/DstMount/SomePath/DstPackage.DstPackage"));
+			CHECK(Patcher.DoPatch(Actual));
+			CHECK_EQUALS(TEXT("SoftObjectPath patching"), Actual, Expected);
+			CHECK_EQUALS(TEXT("SoftObject patching doesn't implicitly update the NameTable"), Patcher.NameTable[0], SrcPackagePathFName);
+			CHECK_EQUALS(TEXT("SoftObject patching doesn't implicitly update the PackageFileSummary"), Patcher.Summary.NameCount, OriginalNameTableCount);
+			Patcher.PatchNameTable();
+			CHECK_EQUALS(TEXT("SoftObject patching updates NameTable entry"), Patcher.NameTable[0], DstPackagePathFName);
+			CHECK_EQUALS(TEXT("SoftObject patching doesn't implicitly update the PackageFileSummary"), Patcher.Summary.NameCount, OriginalNameTableCount);
+		}
+
+		{
+			ResetPatcher();
+
+			FSoftObjectPath Actual(TEXT("/SrcMount/SomePath/SrcPackage.SrcPackage:Some.SrcPackage.Subobject"));
+			// Note we do not replace the sub-object "SrcPackage" despite it matching the original package and object name
+			FSoftObjectPath Expected(TEXT("/DstMount/SomePath/DstPackage.DstPackage:Some.SrcPackage.Subobject"));
+			CHECK(Patcher.DoPatch(Actual));
+			CHECK_EQUALS(TEXT("SoftObjectPath with sub-object path patching"), Actual, Expected);
+			CHECK_EQUALS(TEXT("SoftObject patching doesn't implicitly update the NameTable"), Patcher.NameTable[0], SrcPackagePathFName);
+			CHECK_EQUALS(TEXT("SoftObject patching doesn't implicitly update the PackageFileSummary"), Patcher.Summary.NameCount, OriginalNameTableCount);
+			Patcher.PatchNameTable();
+			CHECK_EQUALS(TEXT("SoftObject patching updates NameTable entry"), Patcher.NameTable[0], DstPackagePathFName);
+			CHECK_EQUALS(TEXT("SoftObject patching doesn't implicitly update the PackageFileSummary"), Patcher.Summary.NameCount, OriginalNameTableCount);
+		}
+	}
+
+	SECTION("DoPatch(FTopLevelAssetPath")
+	{
+		{
+			ResetPatcher();
+
+			FTopLevelAssetPath Actual(SrcPackagePath, SrcAssetFName);
+			FTopLevelAssetPath Expected(DstPackagePath, DstAssetFName);
+			CHECK(Patcher.DoPatch(Actual));
+			CHECK_EQUALS(TEXT("TopLevelAssetPatch(FName,FName) patching"), Actual, Expected);
+			CHECK_EQUALS(TEXT("TopLevelAssetPatch(FName,FName) patching doesn't implicitly update the NameTable"), Patcher.NameTable[0], SrcPackagePathFName);
+			CHECK_EQUALS(TEXT("TopLevelAssetPatch(FName,FName) patching doesn't implicitly update the PackageFileSummary"), Patcher.Summary.NameCount, OriginalNameTableCount);
+			Patcher.PatchNameTable();
+			CHECK_EQUALS(TEXT("TopLevelAssetPatch(FName,FName) patching updates NameTable entry"), Patcher.NameTable[0], DstPackagePathFName);
+			CHECK_EQUALS(TEXT("TopLevelAssetPatch(FName,FName) patching doesn't implicitly update the PackageFileSummary"), Patcher.Summary.NameCount, OriginalNameTableCount);
+		}
+
+		{
+			ResetPatcher();
+
+			FTopLevelAssetPath Actual(SrcPackageObjectPath);
+			FTopLevelAssetPath Expected(DstPackageObjectPath);
+			CHECK(Patcher.DoPatch(Actual));
+			CHECK_EQUALS(TEXT("TopLevelAssetPatch(string) patching"), Actual, Expected);
+			CHECK_EQUALS(TEXT("TopLevelAssetPatch(string) patching doesn't implicitly update the NameTable"), Patcher.NameTable[0], SrcPackagePathFName);
+			CHECK_EQUALS(TEXT("TopLevelAssetPatch(string) patching doesn't implicitly update the PackageFileSummary"), Patcher.Summary.NameCount, OriginalNameTableCount);
+			Patcher.PatchNameTable();
+			CHECK_EQUALS(TEXT("TopLevelAssetPatch(string) patching updates NameTable entry"), Patcher.NameTable[0], DstPackagePathFName);
+			CHECK_EQUALS(TEXT("TopLevelAssetPatch(string) patching doesn't implicitly update the PackageFileSummary"), Patcher.Summary.NameCount, OriginalNameTableCount);
+		}
+	}
+
+	SECTION("DoPatch(FGatherableTextData")
+	{
+		{
+			ResetPatcher();
+
+			FGatherableTextData Actual;
+			Actual.NamespaceName = SrcPackagePath;
+			Actual.SourceData.SourceString = SrcPackagePath;
+			FTextSourceSiteContext SrcSiteContext;
+			SrcSiteContext.KeyName = SrcPackagePath;
+			SrcSiteContext.SiteDescription = SrcPackagePath;
+			Actual.SourceSiteContexts.Add(SrcSiteContext);
+
+			FGatherableTextData Expected = Actual;
+			Expected.SourceSiteContexts = TArray<FTextSourceSiteContext>();
+			FTextSourceSiteContext DstSiteContext;
+			DstSiteContext.KeyName = SrcPackagePath;
+			DstSiteContext.SiteDescription = DstPackagePath;
+			Expected.SourceSiteContexts.Add(DstSiteContext);
+
+			CHECK(Patcher.DoPatch(Actual));
+			CHECK_EQUALS(TEXT("FGatherableTextData patching doesn't update NamespaceName"), Actual.NamespaceName, Expected.NamespaceName);
+			CHECK_EQUALS(TEXT("FGatherableTextData patching doesn't update SourceData.SourceString"), Actual.SourceData.SourceString, Expected.SourceData.SourceString);
+			CHECK_EQUALS(TEXT("FGatherableTextData patching doesn't update SourceSiteContexts[].KeyName"), Actual.SourceSiteContexts[0].KeyName, Expected.SourceSiteContexts[0].KeyName);
+			CHECK_EQUALS(TEXT("FGatherableTextData patching does update SourceData.SourceString[].SiteDescription"), Actual.SourceSiteContexts[0].SiteDescription, Expected.SourceSiteContexts[0].SiteDescription);
+			CHECK_EQUALS(TEXT("FGatherableTextData patching doesn't implicitly update the PackageFileSummary"), Patcher.Summary.NameCount, OriginalNameTableCount);
+			CHECK_EQUALS(TEXT("FGatherableTextData patching doesn't implicitly update the NameTable"), Patcher.NameTable[0], SrcPackagePathFName);
+			CHECK_EQUALS(TEXT("FGatherableTextData patching doesn't implicitly update the NameTable"), Patcher.NameTable[1], SrcAssetFName);
+			Patcher.PatchNameTable();
+			// FGatherableTexData doesn't contain FNames so we shouldn't have updated the NameTable at all
+			CHECK_EQUALS(TEXT("FGatherableTextData patching doesn't implicitly update the PackageFileSummary"), Patcher.Summary.NameCount, OriginalNameTableCount);
+			CHECK_EQUALS(TEXT("FGatherableTextData patching updates NameTable entry"), Patcher.NameTable[0], SrcPackagePathFName);
+			CHECK_EQUALS(TEXT("FGatherableTextData patching doesn't implicitly update the NameTable"), Patcher.NameTable[1], SrcAssetFName);
+		}
+	}
+
+	SECTION("DoPatch(FObjectResource)")
+	{
+		{
+			ResetPatcher();
+
+			CHECK(Patcher.ExportTable.Num() > 0);
+			FObjectResource ExportResource = Patcher.ExportTable[0];
+			CHECK(ExportResource.OldClassName != NAME_None);
+
+			FName Actual = ExportResource.ObjectName;
+			const FName Expected = DstExportObjectFName;
+			CHECK(Actual == SrcExportObjectFName);
+			CHECK(Patcher.NameTable.Contains(Actual));
+
+			CHECK(Patcher.DoPatch(ExportResource, true, Actual));
+			CHECK(ExportResource.OldClassName == NAME_None);
+			CHECK(Actual == Expected);
+		}
+
+		// We are looking at an Export object that has the same name as an export that we will patch
+		// however, in this instance the outers are different. This export should not be patched as a 
+		// result since we do not have a remapping for it's full object path
+		{
+			ResetPatcher();
+
+			CHECK(Patcher.ExportTable.Num() > 1);
+			FObjectResource ExportResource = Patcher.ExportTable[1];
+			CHECK(ExportResource.OldClassName != NAME_None);
+			CHECK(ExportResource.OuterIndex.IsImport());
+
+			FName Actual = ExportResource.ObjectName;
+			const FName Expected = Actual;
+			CHECK(Actual == SrcExportObjectFName);
+			CHECK(Patcher.NameTable.Contains(Actual));
+
+			CHECK(!Patcher.DoPatch(ExportResource, true, Actual));
+			CHECK(ExportResource.OldClassName == NAME_None); // We always clear this
+			CHECK(Actual == Expected);
+		}
+	}
+
+	CHECK(FCoreRedirects::RemoveRedirectList(Context.GetRedirects(), TEXT("Asset Header Patcher Tests")));
+}
+
+#endif // WITH_TESTS
diff --git a/Engine/Source/Developer/AssetTools/Private/AssetHeaderPatcher.h b/Engine/Source/Developer/AssetTools/Private/AssetHeaderPatcher.h
deleted file mode 100644
diff --git a/Engine/Source/Developer/AssetTools/Private/AssetTools.cpp b/Engine/Source/Developer/AssetTools/Private/AssetTools.cpp
--- a/Engine/Source/Developer/AssetTools/Private/AssetTools.cpp
+++ b/Engine/Source/Developer/AssetTools/Private/AssetTools.cpp
@@ -2245,235 +2245,6 @@ void UAssetToolsImpl::GetAllAdvancedCopySources(FName SelectedPackage, FAdvanced
 	}
 }
 
-namespace 
-{
-bool SplitLongPackageName(FStringView LongPackageName, FStringView& PackageRoot, FStringView& PackagePath, FStringView& PackageName)
-{
-	if (LongPackageName.IsEmpty() || LongPackageName[0] != TEXT('/'))
-	{
-		return false;
-	}
-
-	PackageRoot = FStringView(LongPackageName.GetData() + 1); // + 1 to skip the leading '/'
-	int32 SeparatorPos;
-	if (!PackageRoot.FindChar(TEXT('/'), SeparatorPos))
-	{
-		return false;
-	}
-	PackageRoot.LeftInline(SeparatorPos);
-
-	const int32 PackagePathOffset = PackageRoot.Len() + 2; // + 2 for the leading and trailing '/'
-	if (LongPackageName.Len() < PackagePathOffset || !LongPackageName.FindLastChar(TEXT('/'), SeparatorPos))
-	{
-		return false;
-	}
-
-	// May be empty. If the PackageName is off the root there is no PackagePath
-	const int32 PackagePathLen = SeparatorPos - (PackagePathOffset - 1);
-	check(PackagePathLen >= 0);
-	PackagePath = FStringView(LongPackageName.GetData() + PackagePathOffset, PackagePathLen - !!PackagePathLen);
-
-	const int32 PackageNameOffset = PackagePathOffset + PackagePath.Len() + !PackagePath.IsEmpty();
-	PackageName = FStringView(LongPackageName.GetData() + PackageNameOffset, LongPackageName.Len() - PackageNameOffset);
-
-	return true;
-}
-
-TMap<FString, FString> AllSourceAndDestPackages(const TMap<FString, FString>& SourceAndDestPackages)
-{
-	// Paths under the __External root drop the package root, so create mappings, per plugin, 
-	// we can leverage when handling those cases where the package path may have been remapped
-	TMap<FString, TMap<FString, FString>> PluginExternalMappings;
-	for (const TPair<FString, FString>& SrcDstPair : SourceAndDestPackages)
-	{
-		const FString& Src = SrcDstPair.Key;
-		const FString& Dst = SrcDstPair.Value;
-
-		FStringView SrcPackageRoot;
-		FStringView SrcPackagePath;
-		FStringView SrcPackageName;
-		SplitLongPackageName(Src, SrcPackageRoot, SrcPackagePath, SrcPackageName);
-
-		FStringView DstPackageRoot;
-		FStringView DstPackagePath;
-		FStringView DstPackageName;
-		SplitLongPackageName(Dst, DstPackageRoot, DstPackagePath, DstPackageName);
-
-		TMap<FString, FString>& ExternalMappings = PluginExternalMappings.FindOrAddByHash(GetTypeHash(SrcPackageRoot), FString(SrcPackageRoot));
-		FStringView SrcPath = SrcPackagePath.IsEmpty() ? SrcPackageName : SrcPackagePath;
-		FStringView DstPath = DstPackagePath.IsEmpty() ? DstPackageName : DstPackagePath;
-		ExternalMappings.Add(FString(SrcPath), FString(DstPath));
-
-		// if there is a path
-		if (!SrcPackagePath.IsEmpty())
-		{
-			// add the local path/asset for the case of maps (which we cannot tell at this point)
-			ExternalMappings.Add(FString(SrcPath.GetData()), FString(DstPath.GetData()));
-		}
-	}
-
-	TMap<FString, FString> Result;
-	IAssetRegistry& Registry = *IAssetRegistry::Get();
-
-	TArray< TTuple<FString, FString> > ToProcess;
-	Algo::Copy(SourceAndDestPackages, ToProcess);
-
-	TStringBuilder<NAME_SIZE> SrcDependencyBuilder;
-	while (ToProcess.Num())
-	{
-		TTuple<FString, FString> Package = ToProcess.Pop();
-
-		if (Result.Contains(Package.Key))
-		{
-			continue;
-		}
-
-		// Become a patching name even if it doesn't have a file.
-		Result.Add({ Package.Key, Package.Value });
-
-		TArray<FName> Dependencies;
-		if (!Registry.GetDependencies(FName(*Package.Key), Dependencies))
-		{
-			continue;
-		}
-
-		FStringView SrcPackageRoot = FPackageName::SplitPackageNameRoot(Package.Key, nullptr);
-		FStringView DstPackageRoot = FPackageName::SplitPackageNameRoot(Package.Value, nullptr);
-		for (const FName Dependency : Dependencies)
-		{
-			Dependency.ToString(SrcDependencyBuilder);
-			FStringView SrcDependency = SrcDependencyBuilder.ToView();
-
-			if (SourceAndDestPackages.FindByHash(GetTypeHash(SrcDependency), SrcDependency))
-			{
-				// We already handled this mapping
-				continue;
-			}
-
-			FStringView SrcDependencyPackageRoot;
-			FStringView SrcDependencyPackagePath;
-			FStringView SrcDependencyPackageName;
-			SplitLongPackageName(SrcDependency, SrcDependencyPackageRoot, SrcDependencyPackagePath, SrcDependencyPackageName);
-			check(!SrcDependencyPackageRoot.IsEmpty());
-
-			// Only consider dependency paths that are for the same package as our src->dst mapping
-			// If the src mapping doesn't begin with a '/' the package name will be empty, since the path isn't a package path
-			if (SrcDependencyPackageRoot != SrcPackageRoot)
-			{
-				continue;
-			}
-
-			TStringBuilder<NAME_SIZE> DstDependencyString;
-
-			// Special handling for external references. The __External[Actors__|Objects__] directory is always under the package root, may contain an
-			// arbitrary amount of subdirs but then ends with two hash subdirs. The path between the __External[Actors__|Objects__] and the two hash dirs
-			// may need remapping so we look at our external mappings to do so.
-			bool bHasExternalActorDir = SrcDependencyPackagePath.StartsWith(FPackagePath::GetExternalActorsFolderName());
-			bool bHasExternalObjectsDir = !bHasExternalActorDir && SrcDependencyPackagePath.StartsWith(FPackagePath::GetExternalObjectsFolderName());
-			if (bHasExternalActorDir || bHasExternalObjectsDir)
-			{
-				int32 RightPartStartPos;
-				if (!SrcDependencyPackagePath.FindChar(TEXT('/'), RightPartStartPos))
-				{
-					// This is a path to only the special directory, skip it no remapping is needed
-					continue;
-				}
-				RightPartStartPos++; // Skip past the '/'
-
-				// Find the start of the two hash dirs
-				// e.g. __ExternalActors__/path/of/interest/A/A9, we only want 'path/of/interest'
-				FStringView ExternalPackagePath(SrcDependencyPackagePath.GetData() + RightPartStartPos, SrcDependencyPackagePath.Len() - RightPartStartPos);
-				int32 HashDirStartPos = 0;
-				int NumHashDirsToStrip = 2;
-				while (NumHashDirsToStrip--)
-				{
-					if (ExternalPackagePath.FindLastChar(TEXT('/'), HashDirStartPos))
-					{
-						ExternalPackagePath.LeftChopInline(ExternalPackagePath.Len() - HashDirStartPos);
-					}
-				}
-
-				// Our __External[Actors|Objects]__ path is malformed
-				if (HashDirStartPos == INDEX_NONE)
-				{
-					continue;
-				}
-
-				const int32 HashPathOffset = RightPartStartPos + HashDirStartPos;
-				FStringView HashPath(SrcDependencyPackagePath.GetData() + HashPathOffset, SrcDependencyPackagePath.Len() - HashPathOffset);
-				const TMap<FString, FString>* ExternalMappings = PluginExternalMappings.FindByHash(GetTypeHash(SrcPackageRoot), SrcPackageRoot);
-				if (!ExternalMappings)
-				{
-					// We have no mapping for this dependency's external actors/objects
-					continue;
-				}
-				const FString* DstExternalPackagePath = ExternalMappings->FindByHash(GetTypeHash(ExternalPackagePath), ExternalPackagePath);
-								
-				DstDependencyString.AppendChar(TEXT('/'));
-				DstDependencyString.Append(DstPackageRoot);
-				DstDependencyString.AppendChar(TEXT('/'));
-				DstDependencyString.Append(bHasExternalActorDir ? FPackagePath::GetExternalActorsFolderName() : FPackagePath::GetExternalObjectsFolderName());
-				DstDependencyString.AppendChar(TEXT('/'));
-				DstDependencyString.Append(DstExternalPackagePath ? *DstExternalPackagePath : ExternalPackagePath);
-				DstDependencyString.Append(HashPath); // HashPath already contains the leading '/'
-				DstDependencyString.AppendChar(TEXT('/'));
-				DstDependencyString.Append(SrcDependencyPackageName);
-			}
-			else
-			{ 
-				// We aren't handling a special directory so replace the package root
-				DstDependencyString.AppendChar(TEXT('/'));
-				DstDependencyString.Append(DstPackageRoot);
-				DstDependencyString.AppendChar(TEXT('/'));
-
-				if (!SrcDependencyPackagePath.IsEmpty())
-				{
-					DstDependencyString.Append(SrcDependencyPackagePath);
-					DstDependencyString.AppendChar(TEXT('/'));
-				}
-
-				DstDependencyString.Append(SrcDependencyPackageName);
-			}
-
-			// If a dep start with the package name, then we are going to copy the asset.
-			// but we need to recurse on this asset as it may have sub dependencies we don't know of yet.
-			ToProcess.Add({ FString(SrcDependency), DstDependencyString.ToString()});
-		}
-	}
-
-	return Result;
-}
-
-// Returns the index of the '_' character found preceding a series of numbers at the end 
-// of a string (i.e. FName number syntax 'namestring_111'). The index is the equivalent to the "PlainString" length.
-// If the string is a malformed numbered FName then INDEX_NONE is returned.
-int32 EndsWithStrippableNumber(const FStringView Name)
-{
-	// Walk backwards verifying we have only numbers until we find a '_'
-	int32 Pos = Name.Len();
-	while (Pos--)
-	{
-		// Uses sign conversion to check if Name[Pos] is within the inclusive bounds of 0->9
-		if ((uint32)(TEXT('9') - Name[Pos]) > (TEXT('9') - TEXT('0')))
-		{
-			// If we didn't find a number, so exit the loop. 
-			// The Name was a numbered name if the non-number we just found 
-			// is '_' but we must ensure we aren't looking at a name
-			// that looks like "_123" (which would be an invalid FName) or "somename_"
-			if (Name[Pos] != TEXT('_') || Pos == 0 || Pos == Name.Len() - 1)
-			{
-				return INDEX_NONE;
-			}
-
-			return Pos;
-		}
-	}
-
-	// If we get here, the whole Name is a number
-	return INDEX_NONE;
-}
-}
-
 bool UAssetToolsImpl::AdvancedCopyPackages(
 	const TMap<FString, FString>& SourceAndDestPackages,
 	const bool bForceAutosave,
@@ -2510,72 +2281,19 @@ bool UAssetToolsImpl::AdvancedCopyPackages(
 		{
 			TRACE_CPUPROFILER_EVENT_SCOPE(AdvancedCopyPackages.HeaderPatching);
 
-			std::atomic<int32> PatchAssetsCompletedCount = 0;
-			UE::Tasks::FTaskEvent PatchAssetsCompletionTask{ UE_SOURCE_LOCATION };
-
-			TMap<FString, FString> ToCopyAndPatchPackages = AllSourceAndDestPackages(SourceAndDestPackages);
-			TMap<FString, FString> PatchingPatterns = GetAdditionalPatchCopyMappings(SourceAndDestPackages);
-			PatchingPatterns.Append(ToCopyAndPatchPackages);
-
-			// Construct all filenames
-			TMap<FString, FString> ToCopyAndPatchFiles;
-			ToCopyAndPatchFiles.Reserve(ToCopyAndPatchPackages.Num());
-			for (const TTuple<FString, FString>& Package : ToCopyAndPatchPackages)
-			{
-				const FString& PackageName = Package.Key;
-				const FString& DestPackage = Package.Value;
-				FString SrcFilename;
-
-				if (FPackageName::IsVersePackage(PackageName))
-				{
-					// Verse packages are not header patchable.
-					// They are also not Packages as far as DoesPackageExist tells me.
-					// But they are real files that in template copying have already been done, so we dont want a warning message.
-					continue;
-				}
-
-				if (FPackageName::DoesPackageExist(PackageName, &SrcFilename))
-				{
-					FString DestFilename = FPackageName::LongPackageNameToFilename(DestPackage, FString(FPathViews::GetExtension(SrcFilename, true)));
-					ToCopyAndPatchFiles.Add({ MoveTemp(SrcFilename), MoveTemp(DestFilename) });
-				} 
-				else
-				{
-					UE_LOG(LogAssetTools, Warning, TEXT("{%s} package does not exist, and will not be copied."), *PackageName);
-				}
-			}
+			FAssetHeaderPatcher::FContext PatcherContext(SourceAndDestPackages);
+			FAssetHeaderPatcher Patcher(MoveTemp(PatcherContext));
 
-			LoopProgress = MakeUnique<FScopedSlowTask>(static_cast<float>(ToCopyAndPatchFiles.Num()), LOCTEXT("AdvancedCopyPackages.ReplacingAssetReferences", "Replacing Asset References..."));
+			int32 NumFilesToPatch = 0;
+			int32 PatchAssetsCompletedCount = 0;
+			UE::Tasks::FTask PatchAssetsCompletionTask = Patcher.PatchAsync(&NumFilesToPatch, &PatchAssetsCompletedCount);
+		
+			LoopProgress = MakeUnique<FScopedSlowTask>(static_cast<float>(NumFilesToPatch), LOCTEXT("AdvancedCopyPackages.ReplacingAssetReferences", "Replacing Asset References..."));
 			LoopProgress->MakeDialog();
 
-			TSet<FString> ErroredFiles;
-			FCriticalSection  ErroredFilesLock;
-
-			// Spawn tasks (Scatter)
-			for (const TTuple<FString, FString>& Filename : ToCopyAndPatchFiles)
-			{
-				const FString& SrcFilename = Filename.Key;
-				const FString& DestFilename = Filename.Value;
-
-				UE::Tasks::FTask PatcherTask = UE::Tasks::Launch(UE_SOURCE_LOCATION,
-					[&PatchAssetsCompletedCount, &PatchingPatterns, InSrcFilename = SrcFilename, InDestFilename = DestFilename, &ErroredFilesLock, &ErroredFiles] () 
-					{
-						FAssetHeaderPatcher::EResult Result = FAssetHeaderPatcher::DoPatch(InSrcFilename, InDestFilename, PatchingPatterns, /* bBespokeSearchInUse */false);
-						if (Result != FAssetHeaderPatcher::EResult::Success) 
-						{
-							FScopeLock Lock(&ErroredFilesLock);
-							ErroredFiles.Add(InSrcFilename);
-						}
-						PatchAssetsCompletedCount.fetch_add(1, std::memory_order_relaxed);
-					});
-				PatchAssetsCompletionTask.AddPrerequisites(PatcherTask);
-			}
-
-			// Gather
-			PatchAssetsCompletionTask.Trigger();
 			while (!PatchAssetsCompletionTask.Wait(FTimespan::FromSeconds(0.5)))
 			{
-				LoopProgress->CompletedWork = (float)PatchAssetsCompletedCount.load(std::memory_order_relaxed);
+				LoopProgress->CompletedWork = (float)std::atomic_ref(PatchAssetsCompletedCount).load(std::memory_order_relaxed);
 				LoopProgress->TickProgress();
 			}
 
@@ -2583,22 +2301,18 @@ bool UAssetToolsImpl::AdvancedCopyPackages(
 			// And reporting to the user
 			FMessageLog AdvancedCopyLog("AssetTools");
 
-			bool bHasErrors = ErroredFiles.Num() != 0;
-
+			bool bHasErrors = Patcher.HasErrors();
+			TMap<FString, FAssetHeaderPatcher::EResult> ErroredFiles;
 			if (bHasErrors)
 			{
+				ErroredFiles = Patcher.GetErrorFiles();
 				AdvancedCopyLog.NewPage(LOCTEXT("AdvancedCopyPackages_SourceControlErrorsListPage", "Revision Control Errors"));
 			}
 
 			// reporting files with copy errors and filtering successful ones
-			for (const TTuple<FString, FString>& Filename : ToCopyAndPatchFiles)
+			for (const auto& FilenameAndResult : ErroredFiles)
 			{			
-				if (ErroredFiles.Contains(Filename.Key))
-				{
-					AdvancedCopyLog.Error(FText::Format(LOCTEXT("AdvancedCopyPackages_CouldNotProcessError", "{0} could not be processed"), FText::FromString(*Filename.Key)));
-					continue;
-				}
-				SuccessfullyCopiedDestinationFiles.Add(Filename.Value);
+				AdvancedCopyLog.Error(FText::Format(LOCTEXT("AdvancedCopyPackages_CouldNotProcessError", "{0} could not be processed: {1}"), FText::FromString(FilenameAndResult.Key), FText::FromString(*LexToString(FilenameAndResult.Value))));
 			}
 
 			if (SuccessfullyCopiedDestinationFiles.Num() > 0
@@ -2917,250 +2631,6 @@ bool UAssetToolsImpl::AdvancedCopyPackages(const FAdvancedCopyParams& CopyParams
 	return bResult;
 }
 
-/** Copies a file, patching internal references without performing a de-serialization. This is a blocking operation. returns true on successful copy */
-bool UAssetToolsImpl::PatchCopyPackageFile(const FString& SrcFile, const FString& DstFile, const TMap<FString, FString>& SearchForAndReplace) const
-{
-	FAssetHeaderPatcher::EResult Result = FAssetHeaderPatcher::DoPatch(SrcFile, DstFile, SearchForAndReplace, /* bBespokeSearchInUse */ true);
-	return (Result == FAssetHeaderPatcher::EResult::Success);
-}
-
-TMap<FString, FString> UAssetToolsImpl::GetPatchCopyMappingsForRootRename(
-	const FString& SrcRoot,
-	const FString& DstRoot,
-	const FString& SrcBaseDir,
-	const TArray<TPair<FString, FString>>& SourceAndDestFiles,
-	const TMap<FString, FString>& MountPointReplacements) const
-{
-	TMap<FString, FString> Result;
-	Result.Reserve(6 + SourceAndDestFiles.Num() + MountPointReplacements.Num()); // usually only 3, +6 just in case
-
-	{	// Plugin name patterns
-		FString SrcPath = FPaths::Combine(TEXT("/"), SrcRoot, SrcRoot);
-		FString DstPath = FPaths::Combine(TEXT("/"), DstRoot, SrcRoot);
-
-		FName SrcRootName = *SrcRoot;
-		if (SrcRootName.GetDisplayNameEntry()->GetNameLength() != SrcRoot.Len())
-		{
-			// SrcRoot has a '_$Number' tail, so it looks like a FName
-			// IE: SrcRoot = Src_1
-			FString ShortSrcRoot = SrcRootName.GetPlainNameString(); // Src
-			FString ShortSrcPath = FPaths::Combine(TEXT("/"), SrcRoot, ShortSrcRoot);
-
-			Result.Add(SrcPath + TEXT(".") + ShortSrcRoot, DstPath + TEXT(".") + SrcRoot);	// /Src_1/Src_1.Src -> /Dst/Src_1.Src_1
-			Result.Add(MoveTemp(ShortSrcPath), DstPath);							        // /Src_1/Src       -> /Dst/Src_1
-
-			// Root package patching has a special case for the GameFeatureData file,
-			// Specifically the PrimaryAssetName in it Tag Data inside the Asset Registry.
-			// This is the only place we need the Root -> Root rename so we target that Key inside that file explicitly
-			// with this name decoration which the FAssetHeaderPatcher will use for just this case.
-			Result.Add(TEXT("<GameFeatureData.PrimaryAssetName>") + ShortSrcRoot, DstRoot);	// <GameFeatureData.PrimaryAssetName>Src -> Dst (for matching in specific modes)
-		}
-
-		Result.Add(SrcPath + TEXT(".") + SrcRoot, DstPath + TEXT(".") + SrcRoot);	// /Src/Src.Src -> /Dst/Src.Src
-		Result.Add(MoveTemp(SrcPath), MoveTemp(DstPath));							// /Src/Src     -> /Dst/Src
-		Result.Add(TEXT("<GameFeatureData.PrimaryAssetName>") + SrcRoot, DstRoot);	// <GameFeatureData.PrimaryAssetName>Src -> Dst (for matching in specific modes)
-	}
-
-	const FString SourceContentPath = FPaths::Combine(SrcBaseDir, TEXT("Content"));
-
-	for (const TTuple<FString, FString>& SourceAndDest : SourceAndDestFiles)
-	{
-		const FString& SrcFileName = SourceAndDest.Key;
-
-		if (FPaths::IsUnderDirectory(SrcFileName, SourceContentPath))
-		{
-			if (FStringView RelativePkgPath; FPathViews::TryMakeChildPathRelativeTo(SrcFileName, SourceContentPath, RelativePkgPath))
-			{
-				RelativePkgPath = FPathViews::GetBaseFilenameWithPath(RelativePkgPath); // chop the extension
-				if (RelativePkgPath.Len() > 0 && !RelativePkgPath.EndsWith(TEXT("/")))
-				{
-					Result.Add(FPaths::Combine(TEXT("/"), SrcRoot, RelativePkgPath),
-						       FPaths::Combine(TEXT("/"), DstRoot, RelativePkgPath));
-				}
-			}
-		}
-	}
-
-	// Mount point replacements.
-	// In the case where a project has multiple plugins, they are represented as different package roots (mountpoints)
-	// So a Package my reference an asset, or object in a different mount point.
-	// We tag this information with '<Mountpoint>' so the header patcher can do a partial replacement (instead of the normal complete replacement)
-	// if the mount point is detected.
-	for (const TPair<FString, FString>& MountPair : MountPointReplacements)
-	{
-		Result.Add(TEXT("<Mountpoint>") + MountPair.Key, MountPair.Value);
-	}
-
-	return Result;
-}
-
-TMap<FString, FString> UAssetToolsImpl::GetAdditionalPatchCopyMappings(const TMap<FString, FString>& SourceAndDestPackages) const
-{
-	TMap<FString, FString> Result;
-
-	TStringBuilder<24> ExternalActorsFolderBuilder;
-	ExternalActorsFolderBuilder << FPackagePath::GetExternalActorsFolderName() << TEXT("/");
-	const FStringView ExternalActorsFolder = ExternalActorsFolderBuilder.ToView();
-
-	TStringBuilder<24> ExternalObjectsFolderBuilder;
-	ExternalObjectsFolderBuilder << FPackagePath::GetExternalObjectsFolderName() << TEXT("/");
-	const FStringView ExternalObjectsFolder = ExternalObjectsFolderBuilder.ToView();
-
-	TStringBuilder<NAME_SIZE> NameBuilder;
-	for (const TTuple<FString, FString>& Package : SourceAndDestPackages)
-	{
-		const FString& SrcNameString = Package.Key;
-		const FString& DstNameString = Package.Value;
-
-		FStringView DstPackageName = FPathViews::GetBaseFilename(DstNameString);
-		FStringView SrcPackageName;
-		{
-			FStringView SrcPackageRoot;
-			FStringView SrcPackagePath;
-			if (!ensure(SplitLongPackageName(SrcNameString, SrcPackageRoot, SrcPackagePath, SrcPackageName))
-				|| SrcPackagePath.StartsWith(ExternalActorsFolder)
-				|| SrcPackagePath.StartsWith(ExternalObjectsFolder))
-			{
-				continue;
-			}
-		}
-
-		// Path.ObjectName mapping
-		{
-			NameBuilder.Reset();
-			NameBuilder.Append(SrcNameString);
-			NameBuilder.AppendChar(TEXT('.'));
-			NameBuilder.Append(SrcPackageName);
-			FString SrcObjectName = NameBuilder.ToString();
-
-			NameBuilder.Reset();
-			NameBuilder.Append(DstNameString);
-			NameBuilder.AppendChar(TEXT('.'));
-			NameBuilder.Append(DstPackageName);
-			FString DstObjectName = NameBuilder.ToString();
-			Result.Add(MoveTemp(SrcObjectName), MoveTemp(DstObjectName));
-		}
-
-		// For package names with a '_[0-9]+' suffix: Add a mapping with the suffix stripped but 
-		// maintain the full long path provided.
-		// NameTable entries have this suffix removed and we won't match them otherwise.
-		{
-			int32 SrcNameStrPlainStringLen = EndsWithStrippableNumber(SrcNameString);				// PackageRoot_0/PackagePath_0/PackageName_0
-			if (SrcNameStrPlainStringLen != INDEX_NONE)
-			{
-				FStringView SrcNamePlainString(*SrcNameString, SrcNameStrPlainStringLen);			// PackageRoot_0/PackagePath_0/PackageName
-
-				// It's valid to rename a numbered name to a non-numbered name, 
-				// default to the original name and use the non-numbered name if found later
-				FStringView DstNamePlainString = DstNameString;
-				int32 DstNameStrPlainStringLen = EndsWithStrippableNumber(DstNameString);
-				if (DstNameStrPlainStringLen != INDEX_NONE)
-				{
-					DstNamePlainString.LeftInline(DstNameStrPlainStringLen);
-				}
-				Result.Add(FString(SrcNamePlainString), FString(DstNamePlainString));				// PackageRoot_0/PackagePath_0/PackageName => DstPackageRoot_0/DstPackagePath_0/DstPackageName
-
-				FStringView SrcShortPackageName = FPathViews::GetBaseFilename(SrcNamePlainString);	// PackageName
-				NameBuilder.Reset();
-				NameBuilder.Append(SrcNameString);
-				NameBuilder.AppendChar(TEXT('.'));
-				NameBuilder.Append(SrcShortPackageName);
-				FString SrcPathWithPlainObjectName = NameBuilder.ToString();
-
-				FStringView DstShortPackageName = FPathViews::GetBaseFilename(DstNamePlainString);	// PackageName
-				NameBuilder.Reset();
-				NameBuilder.Append(DstNameString);
-				NameBuilder.AppendChar(TEXT('.'));
-				NameBuilder.Append(DstShortPackageName);
-				FString DstPathWithObjectName = NameBuilder.ToString();
-				Result.Add(MoveTemp(SrcPathWithPlainObjectName), MoveTemp(DstPathWithObjectName));	// PackageRoot_0/PackagePath_0/PackageName_0.PackageName => DstPackageRoot_0/DstPackagePath_0/DstPackageName_0.DstPackageName
-			}
-		}
-
-		// VerseAssetPath mapping used in GatherableTextData
-		// e.g. /localhost/Some/Package/Path/PackageName/PackageObject (note the use of '/' instead of '.' for the top-level Package object)
-		{
-			// Using a constant to avoid pulling in, what would otherwise be, unnecessary dependencies
-			const TCHAR* VerseRoot = TEXT("/localhost");
-
-			NameBuilder.Reset();
-			NameBuilder.Append(VerseRoot);
-			NameBuilder.Append(SrcNameString);
-			NameBuilder.AppendChar(TEXT('/'));
-			NameBuilder.Append(SrcPackageName);
-			FString SrcObjectName = NameBuilder.ToString();
-
-			NameBuilder.Reset();
-			NameBuilder.Append(VerseRoot);
-			NameBuilder.Append(DstNameString);
-			NameBuilder.AppendChar(TEXT('/'));
-			NameBuilder.Append(DstPackageName);
-			FString DstObjectName = NameBuilder.ToString();
-			Result.Add(MoveTemp(SrcObjectName), MoveTemp(DstObjectName));
-		}
-
-		if (SrcPackageName != DstPackageName)
-		{
-			// PackageName (without PackageRoot or PackagePath) mapping (numbered suffix included)
-			{
-				Result.Add({ FString(SrcPackageName), FString(DstPackageName) });
-			}
-
-			// PackageName (without PackageRoot or PackagePath) mapping (numbered suffix excluded)
-			// NameTable entries have this suffix removed and we won't match them otherwise.
-			{
-				int32 SrcPackageNamePlainStringLen = EndsWithStrippableNumber(SrcPackageName);						// PackageName_0
-				if (SrcPackageNamePlainStringLen != INDEX_NONE)
-				{
-					FStringView SrcPlainPackageName(SrcPackageName.GetData(), SrcPackageNamePlainStringLen);
-
-					// It's valid to rename a numbered name to a non-numbered name, 
-					// default to the original name and use the non-numbered name if found later
-					FStringView DstPlainPackageName = DstPackageName;
-					int32 DstPackageNamePlainStringLen = EndsWithStrippableNumber(DstPackageName);
-					if (DstPackageNamePlainStringLen != INDEX_NONE)
-					{
-						DstPlainPackageName = FStringView(DstPackageName.GetData(), DstPackageNamePlainStringLen);
-					}
-					Result.Add(FString(SrcPlainPackageName), FString(DstPlainPackageName));							// PackageName => DstPackageName
-				}
-			}
-
-			// Compiled Blueprint class names
-			{
-				NameBuilder.Reset();
-				NameBuilder.Append(SrcPackageName);
-				NameBuilder.Append(TEXT("_C"));
-				FString SrcBlueprintClassName = NameBuilder.ToString();
-
-				NameBuilder.Reset();
-				NameBuilder.Append(DstPackageName);
-				NameBuilder.Append(TEXT("_C"));
-				FString DstBlueprintClassName = NameBuilder.ToString();
-				Result.Add(MoveTemp(SrcBlueprintClassName), MoveTemp(DstBlueprintClassName));
-			}
-
-			// Blueprint generated class default object
-			{
-				NameBuilder.Reset();
-				NameBuilder.Append(DEFAULT_OBJECT_PREFIX);
-				NameBuilder.Append(SrcPackageName);
-				NameBuilder.Append(TEXT("_C"));
-				FString SrcDefaultGeneratedBlueprintClassName = NameBuilder.ToString();
-
-				NameBuilder.Reset();
-				NameBuilder.Append(DEFAULT_OBJECT_PREFIX);
-				NameBuilder.Append(DstPackageName);
-				NameBuilder.Append(TEXT("_C"));
-				FString DstDefaultGeneratedBlueprintClassName = NameBuilder.ToString();
-				Result.Add(MoveTemp(SrcDefaultGeneratedBlueprintClassName), MoveTemp(DstDefaultGeneratedBlueprintClassName));
-			}
-		}
-	}
-
-	return Result;
-}
-
 bool UAssetToolsImpl::IsDiscoveringAssetsInProgress() const
 {
 	FAssetRegistryModule& AssetRegistryModule = FModuleManager::LoadModuleChecked<FAssetRegistryModule>("AssetRegistry");
diff --git a/Engine/Source/Developer/AssetTools/Private/AssetTools.h b/Engine/Source/Developer/AssetTools/Private/AssetTools.h
--- a/Engine/Source/Developer/AssetTools/Private/AssetTools.h
+++ b/Engine/Source/Developer/AssetTools/Private/AssetTools.h
@@ -134,9 +134,6 @@ public:
 	virtual bool ValidateFlattenedAdvancedCopyDestinations(const TMap<FString, FString>& FlattenedPackagesAndDestinations) const override;
 	virtual void GetAllAdvancedCopySources(FName SelectedPackage, FAdvancedCopyParams& CopyParams, TArray<FName>& OutPackageNamesToCopy, TMap<FName, FName>& DependencyMap, const class UAdvancedCopyCustomization* CopyCustomization) const override;
 	virtual void InitAdvancedCopyFromCopyParams(FAdvancedCopyParams CopyParams) const override;
-	virtual bool PatchCopyPackageFile(const FString& SrcFile, const FString& DstFile, const TMap<FString, FString>& SearchForAndReplace) const override;
-	virtual TMap<FString, FString> GetPatchCopyMappingsForRootRename(const FString& SrcRoot, const FString& DstRoot, const FString& SrcBaseDir, const TArray<TPair<FString, FString>>& SourceAndDestFiles, const TMap<FString, FString>& MountPointReplacements) const override;
-	virtual TMap<FString, FString> GetAdditionalPatchCopyMappings(const TMap<FString, FString>& SourceAndDestPackages) const override;
 	virtual void OpenEditorForAssets(const TArray<UObject*>& Assets) override;
 	virtual void ConvertVirtualTextures(const TArray<UTexture2D*>& Textures, bool bConvertBackToNonVirtual, const TArray<UMaterial*>* RelatedMaterials = nullptr) const override;
 	virtual bool IsAssetClassSupported(const UClass* AssetClass) const override;
diff --git a/Engine/Source/Developer/AssetTools/Public/IAssetTools.h b/Engine/Source/Developer/AssetTools/Public/IAssetTools.h
--- a/Engine/Source/Developer/AssetTools/Public/IAssetTools.h
+++ b/Engine/Source/Developer/AssetTools/Public/IAssetTools.h
@@ -612,15 +612,6 @@ public:
 	/* Given a complete set of copy parameters, which includes the selected package set, start the advanced copy process */
 	virtual void InitAdvancedCopyFromCopyParams(FAdvancedCopyParams CopyParams) const = 0;
 
-	/** Copies a file, patching internal references without performing a de-serialization. This is a blocking operation. returns true on successful copy */
-	UE_INTERNAL virtual bool PatchCopyPackageFile(const FString& SrcFile, const FString& DstFile, const TMap<FString, FString>& SearchForAndReplace) const = 0;
-
-	/** Generates the PatchCopyPackageFile SearchForAndReplace parameter if all you are doing is changing the root and not the relative path of assets */
-	UE_INTERNAL virtual TMap<FString, FString> GetPatchCopyMappingsForRootRename(const FString& SrcRoot, const FString& DstRoot, const FString& SrcBaseDir, const TArray<TPair<FString, FString>>& SourceAndDestFiles, const TMap<FString, FString>& MountPointReplacements) const = 0;
-
-	/** Generates additional entries for the PatchCopyPackageFile SearchForAndReplace parameter based on the specified package paths mapping */
-	UE_INTERNAL virtual TMap<FString, FString> GetAdditionalPatchCopyMappings(const TMap<FString, FString>& SourceAndDestPackages) const = 0;
-
 	/** Opens editor for assets */
 	UFUNCTION(BlueprintCallable, Category = "Editor Scripting | Asset Tools", meta = (DeprecatedFunction, DeprecationMessage = "Please use UAssetEditorSubsystem::OpenEditorForAssets instead."))
 	virtual void OpenEditorForAssets(const TArray<UObject*>& Assets) = 0;
diff --git a/Engine/Source/Runtime/Engine/Private/WorldPartition/WorldPartitionActorDescArchive.cpp b/Engine/Source/Runtime/Engine/Private/WorldPartition/WorldPartitionActorDescArchive.cpp
--- a/Engine/Source/Runtime/Engine/Private/WorldPartition/WorldPartitionActorDescArchive.cpp
+++ b/Engine/Source/Runtime/Engine/Private/WorldPartition/WorldPartitionActorDescArchive.cpp
@@ -121,6 +121,13 @@ PRAGMA_ENABLE_DEPRECATION_WARNINGS
 	BaseDescSizeof = BaseDesc ? BaseDesc->GetSizeOf() : 0;
 }
 
+FArchive& FActorDescArchive::operator<<(FTopLevelAssetPath& Value)
+{
+	((FArchive&)*this) << Value;
+
+	return *this;
+};
+
 FArchive& FActorDescArchive::operator<<(FSoftObjectPath& Value)
 {
 	Value.SerializePathWithoutFixup(*this);
@@ -167,6 +174,23 @@ FArchive& FActorDescArchivePatcher::operator<<(FSoftObjectPath& Value)
 	return *this;
 }
 
+FArchive& FActorDescArchivePatcher::operator<<(FTopLevelAssetPath& Value)
+{
+	{
+		TGuardValue<bool> GuardIsPatching(bIsPatching, true);
+		FActorDescArchive::operator<<(Value);
+		AssetDataPatcher->DoPatch(Value);
+	}
+
+	// Only write out values if we aren't already patching since this function can be called
+	// from other patching functions which will perform the final write of the patched values
+	if (!bIsPatching)
+	{
+		OutAr << Value;
+	}
+	return *this;
+}
+
 void FActorDescArchivePatcher::Serialize(void* V, int64 Length)
 {
 	FActorDescArchive::Serialize(V, Length);
diff --git a/Engine/Source/Runtime/Engine/Public/WorldPartition/WorldPartitionActorDesc.h b/Engine/Source/Runtime/Engine/Public/WorldPartition/WorldPartitionActorDesc.h
--- a/Engine/Source/Runtime/Engine/Public/WorldPartition/WorldPartitionActorDesc.h
+++ b/Engine/Source/Runtime/Engine/Public/WorldPartition/WorldPartitionActorDesc.h
@@ -68,6 +68,7 @@ private:
 
 struct FWorldPartitionAssetDataPatcher
 {
+	virtual ~FWorldPartitionAssetDataPatcher() {}
 	virtual bool DoPatch(FString& InOutString) = 0;
 	virtual bool DoPatch(FName& InOutName) = 0;
 	virtual bool DoPatch(FSoftObjectPath& InOutSoft) = 0;
diff --git a/Engine/Source/Runtime/Engine/Public/WorldPartition/WorldPartitionActorDescArchive.h b/Engine/Source/Runtime/Engine/Public/WorldPartition/WorldPartitionActorDescArchive.h
--- a/Engine/Source/Runtime/Engine/Public/WorldPartition/WorldPartitionActorDescArchive.h
+++ b/Engine/Source/Runtime/Engine/Public/WorldPartition/WorldPartitionActorDescArchive.h
@@ -28,6 +28,9 @@ public:
 	virtual FArchive& operator<<(FWeakObjectPtr& Value) override { unimplemented(); return *this; }
 	//~ End FArchive Interface
 
+	using FArchive::operator<<;
+	virtual FArchive& operator<<(FTopLevelAssetPath& Value); 
+
 	template <typename DestPropertyType, typename SourcePropertyType>
 	struct TDeltaSerializer
 	{
@@ -140,6 +143,8 @@ public:
 	virtual void Serialize(void* V, int64 Length) override;
 	//~ End FArchive Interface
 
+	virtual FArchive& operator<<(FTopLevelAssetPath& Value) override;
+
 private:
 	FArchive& OutAr;
 	FWorldPartitionAssetDataPatcher* AssetDataPatcher;
