Set use HeaderPatching as true by default for AdvancedCopy and CopyProjectWithHeaderPatching

#rnx

diff --git a/Engine/Source/Developer/AssetTools/Private/AssetTools.cpp b/Engine/Source/Developer/AssetTools/Private/AssetTools.cpp
--- a/Engine/Source/Developer/AssetTools/Private/AssetTools.cpp
+++ b/Engine/Source/Developer/AssetTools/Private/AssetTools.cpp
@@ -469,7 +469,7 @@ namespace UE::AssetTools::Private
 	 * CVar to specify if we should use Header patching in advanced copy.
 	 * Default is false.
 	 */
-	bool bEnableHeaderPatching = false;
+	bool bEnableHeaderPatching = true;
 	FAutoConsoleVariableRef CVarEnableHeaderPatching(
 		TEXT("AssetTools.UseHeaderPatchingAdvancedCopy"),
 		bEnableHeaderPatching,
