Set use HeaderPatching as true by default for AdvancedCopy and CopyProjectWithHeaderPatching

#rnx

diff --git a/Engine/Source/Developer/AssetTools/Private/AssetHeaderPatcher.cpp b/Engine/Source/Developer/AssetTools/Private/AssetHeaderPatcher.cpp
--- a/Engine/Source/Developer/AssetTools/Private/AssetHeaderPatcher.cpp
+++ b/Engine/Source/Developer/AssetTools/Private/AssetHeaderPatcher.cpp
@@ -936,7 +936,7 @@ UE::Tasks::FTask FAssetHeaderPatcher::PatchAsync(int32* InOutNumFilesToPatch, in
 	// Note we are scheduling and launching tasks one at a time rather than preparing all jobs and launching all at once.
 	// While this means more overhead scheduling, it means that we won't have many tasks all hit the filesystem at the same time
 	// attempting to read and (more importantly) write to disk at the exact same time.
-	constexpr bool bSingleThreaded = true; // Useful for debugging
+	constexpr bool bSingleThreaded = false; // Useful for debugging
 	for (const TTuple<FString, FString>& Filename : PatchedFiles)
 	{
 		auto DoPatchFn = [this, 
@@ -1953,11 +1953,6 @@ void FAssetHeaderPatcherInner::PatchHeader_PatchSections()
 		}
 	}
 
-/*
-	Enable when serialization is replaced with PackageReader. Thumbnail serialization seems to work but 
-	we don't test it well enough to confidently enable. This data is regenerated when missing upon load 
-	so not patching it isn't a huge deal. Patching it incorrectly may corrupt the whole file however.
-
 	// SearchableNamesMap
 	for (auto& Pair : SearchableNamesMap)
 	{
@@ -1973,7 +1968,6 @@ void FAssetHeaderPatcherInner::PatchHeader_PatchSections()
 	{
 		DoPatch(ThumbnailEntry);
 	}
-*/
 
 	// Asset Register Data
 	for (FAssetRegistryObjectData& ObjData : AssetRegistryData.ObjectData)
diff --git a/Engine/Source/Developer/AssetTools/Private/AssetTools.cpp b/Engine/Source/Developer/AssetTools/Private/AssetTools.cpp
--- a/Engine/Source/Developer/AssetTools/Private/AssetTools.cpp
+++ b/Engine/Source/Developer/AssetTools/Private/AssetTools.cpp
@@ -469,7 +469,7 @@ namespace UE::AssetTools::Private
 	 * CVar to specify if we should use Header patching in advanced copy.
 	 * Default is false.
 	 */
-	bool bEnableHeaderPatching = false;
+	bool bEnableHeaderPatching = true;
 	FAutoConsoleVariableRef CVarEnableHeaderPatching(
 		TEXT("AssetTools.UseHeaderPatchingAdvancedCopy"),
 		bEnableHeaderPatching,
