Fixing long filepaths being incorrectly validated for machines that have longfilepaths enabled

diff --git a/Engine/Source/Developer/AssetTools/Private/AssetViewUtils.cpp b/Engine/Source/Developer/AssetTools/Private/AssetViewUtils.cpp
--- a/Engine/Source/Developer/AssetTools/Private/AssetViewUtils.cpp
+++ b/Engine/Source/Developer/AssetTools/Private/AssetViewUtils.cpp
@@ -1150,10 +1150,10 @@ bool AssetViewUtils::IsValidFolderName(const FString& FolderName, FText& Reason)
 		return false;
 	}
 
-	if ( FolderName.Len() > FPlatformMisc::GetMaxPathLength() )
+	if ( FolderName.Len() > GetMaxCookPathLen())
 	{
 		Reason = FText::Format( LOCTEXT("InvalidFolderName_TooLongForCooking", "Filename '{0}' is too long; this may interfere with cooking for consoles. Unreal filenames should be no longer than {1} characters." ),
-			FText::FromString(FolderName), FText::AsNumber(FPlatformMisc::GetMaxPathLength()) );
+			FText::FromString(FolderName), FText::AsNumber(GetMaxCookPathLen()) );
 		return false;
 	}
 
@@ -1393,7 +1393,7 @@ FLinearColor AssetViewUtils::GetDefaultColor()
 
 namespace AssetViewUtils
 {
-static const TConsoleVariableData<int32>* CVarMaxFullPathLength = IConsoleManager::Get().RegisterConsoleVariable(TEXT("MaxAssetFullPath"), FPlatformMisc::GetMaxPathLength(), TEXT("Maximum full path name of an asset."))->AsVariableInt();
+static const TConsoleVariableData<int32>* CVarMaxFullPathLength = IConsoleManager::Get().RegisterConsoleVariable(TEXT("MaxAssetFullPath"), GetMaxCookPathLen(), TEXT("Maximum full path name of an asset."))->AsVariableInt();
 static const TConsoleVariableData<FString>* CVarExternalPluginCookedRootPath = IConsoleManager::Get().RegisterConsoleVariable(TEXT("ExternalPluginCookedAssetRootPath"), TEXT(""), TEXT("Root path to use when estimating the cooked path external plugin assets, or empty to use the standard engine/project root."))->AsVariableString();
 }
 
@@ -1508,11 +1508,11 @@ bool AssetViewUtils::IsValidFolderPathForCreate(const FString& InFolderPath, con
 	}
 
 	// Make sure we are not creating a folder path that is too long
-	if (NewFolderPath.Len() > FPlatformMisc::GetMaxPathLength() - MAX_CLASS_NAME_LENGTH)
+	if (NewFolderPath.Len() > GetMaxCookPathLen() - MAX_CLASS_NAME_LENGTH)
 	{
 		// The full path for the folder is too long
 		OutErrorMessage = FText::Format(LOCTEXT("RenameFolderPathTooLong", "The full path for the folder is too deep, the maximum is '{0}'. Please choose a shorter name for the folder or create it in a shallower folder structure."),
-			FText::AsNumber(FPlatformMisc::GetMaxPathLength()));
+			FText::AsNumber(GetMaxCookPathLen()));
 		// Return false to indicate that the user should enter a new name for the folder
 		return false;
 	}
