Fixed static-init crash

#rnx

diff --git a/Engine/Source/Developer/AssetTools/Private/AssetViewUtils.cpp b/Engine/Source/Developer/AssetTools/Private/AssetViewUtils.cpp
--- a/Engine/Source/Developer/AssetTools/Private/AssetViewUtils.cpp
+++ b/Engine/Source/Developer/AssetTools/Private/AssetViewUtils.cpp
@@ -1393,8 +1393,19 @@ FLinearColor AssetViewUtils::GetDefaultColor()
 
 namespace AssetViewUtils
 {
-static const TConsoleVariableData<int32>* CVarMaxFullPathLength = IConsoleManager::Get().RegisterConsoleVariable(TEXT("MaxAssetFullPath"), GetMaxCookPathLen(), TEXT("Maximum full path name of an asset."))->AsVariableInt();
+static const TConsoleVariableData<int32>* CVarMaxFullPathLength = IConsoleManager::Get().RegisterConsoleVariable(TEXT("MaxAssetFullPath"), 0, TEXT("Maximum full path name of an asset, or <=0 to use GetMaxCookPathLen()."))->AsVariableInt();
 static const TConsoleVariableData<FString>* CVarExternalPluginCookedRootPath = IConsoleManager::Get().RegisterConsoleVariable(TEXT("ExternalPluginCookedAssetRootPath"), TEXT(""), TEXT("Root path to use when estimating the cooked path external plugin assets, or empty to use the standard engine/project root."))->AsVariableString();
+
+int32 GetMaxFullPathLength()
+{
+	int32 MaxFullPathLength = CVarMaxFullPathLength->GetValueOnGameThread();
+	if (MaxFullPathLength <= 0)
+	{
+		MaxFullPathLength = GetMaxCookPathLen();
+	}
+	return MaxFullPathLength;
+}
+
 }
 
 bool AssetViewUtils::IsValidObjectPathForCreate(const FString& ObjectPath, FText& OutErrorMessage, bool bAllowExistingAsset)
@@ -1437,11 +1448,12 @@ bool AssetViewUtils::IsValidObjectPathForCreate(const FString& ObjectPath, const
 	}
 
 	const FString FullPath = FPaths::ConvertRelativePathToFull(RelativePathFilename);	// path to file on disk
-	if (FullPath.Len() > CVarMaxFullPathLength->GetValueOnGameThread() )
+	if (int32 MaxFullPathLength = GetMaxFullPathLength();
+		FullPath.Len() > MaxFullPathLength)
 	{
 		// The full path for the asset is too long
 		OutErrorMessage = FText::Format(LOCTEXT("AssetPathTooLong", "The absolute file path for the asset is too long, the maximum is '{0}'. \nPlease choose a shorter name for the asset or create it in a shallower folder structure."),
-			FText::AsNumber(CVarMaxFullPathLength->GetValueOnGameThread()));
+			FText::AsNumber(MaxFullPathLength));
 		// Return false to indicate that the user should enter a new name
 		return false;
 	}
