Fix "Update Redirector References" process unable to fix up redirectors of assets used in level instances

diff --git a/Engine/Source/Developer/AssetTools/Private/AssetFixUpRedirectors.cpp b/Engine/Source/Developer/AssetTools/Private/AssetFixUpRedirectors.cpp
--- a/Engine/Source/Developer/AssetTools/Private/AssetFixUpRedirectors.cpp
+++ b/Engine/Source/Developer/AssetTools/Private/AssetFixUpRedirectors.cpp
@@ -13,6 +13,7 @@
 #include "ICollectionManager.h"
 #include "ISourceControlModule.h"
 #include "ISourceControlOperation.h"
+#include "LevelInstance/LevelInstanceSubsystem.h"
 #include "Logging/MessageLog.h"
 #include "Misc/MessageDialog.h"
 #include "Misc/PackageName.h"
@@ -715,6 +716,27 @@ void FAssetFixUpRedirectors::ExecuteFixUp(TArray<TWeakObjectPtr<UObjectRedirecto
 			return true;
 		}, false, RF_Standalone, EInternalObjectFlags::RootSet);
 	}
+
+	// Reset loaders of assets used in level instances to allow referencing packages to be saved
+	TSet<FName> WorldAssetsNeedingLoadersReset;
+	for (UPackage* Package : ReferencingPackagesToSave)
+	{
+		TArray<FAssetData> ReferencingPackageAssets;
+		if (AssetRegistryModule.Get().GetAssetsByPackageName(Package->GetFName(), ReferencingPackageAssets, /*bIncludeOnlyOnDiskAssets=*/true))
+		{
+			for (const FAssetData& Asset : ReferencingPackageAssets)
+			{
+				if (!Asset.GetOptionalOuterPathName().IsNone())
+				{
+					WorldAssetsNeedingLoadersReset.Add(FSoftObjectPath(Asset.GetOptionalOuterPathName().ToString()).GetLongPackageFName());
+				}
+			}
+		}
+	}
+	for (const FName& WorldAsset : WorldAssetsNeedingLoadersReset)
+	{
+		ULevelInstanceSubsystem::ResetLoadersForWorldAsset(*WorldAsset.ToString());
+	}
 	
 	// Check out all referencing packages, leave redirectors for assets referenced by packages that are not checked out and remove those packages from the save list.
 	bool bUserAcceptedCheckout = true; // If source control is disabled, assume checkout was selected
