- Adds GRedirectCollector support to FSoftObjectPath patching in the AssetHeaderPatcher. Explicit mappings from the user are honoured first and then global redirectors are applied.
- Corrects improperly wrapped string literal with TEXT("") in AssetHeaderPatcher

#rnx

diff --git a/Engine/Source/Developer/AssetTools/Private/AssetHeaderPatcher.cpp b/Engine/Source/Developer/AssetTools/Private/AssetHeaderPatcher.cpp
--- a/Engine/Source/Developer/AssetTools/Private/AssetHeaderPatcher.cpp
+++ b/Engine/Source/Developer/AssetTools/Private/AssetHeaderPatcher.cpp
@@ -14,6 +14,7 @@
 #include "Misc/PackageName.h"
 #include "Misc/Paths.h"
 #include "Misc/PathViews.h"
+#include "Misc/RedirectCollector.h"
 #include "Serialization/LargeMemoryReader.h"
 #include "UObject/CoreRedirects.h"
 #include "UObject/Linker.h"
@@ -135,7 +136,7 @@ FAssetHeaderPatcher::FContext::FContext(const FString& InSrcRoot, const FString&
 void FAssetHeaderPatcher::FContext::AddVerseMounts()
 {
 	// Todo: Expose this so callers provide this data
-	VerseMountPoints.Add("localhost");
+	VerseMountPoints.Add(TEXT("localhost"));
 }
 
 void FAssetHeaderPatcher::FContext::GenerateFilePathsFromPackagePaths()
@@ -1192,6 +1193,10 @@ FAssetHeaderPatcher::EResult FAssetHeaderPatcherInner::PatchHeader_Deserialize()
 		SoftObjectPathTable.Reserve(Summary.SoftObjectPathsCount);
 		for (int32 Idx = 0; Idx < Summary.SoftObjectPathsCount; ++Idx)
 		{
+			// Note, a non IsPersistent() archive is used to preserve the original
+			// FSoftObjectPaths found in the header, since those will refer to entries in the 
+			// NameTable. The call to SerializePath below will redirect FNames when using an
+			// IsPersistent archive which would make fixing up the nametable more difficult
 			FSoftObjectPath& PathRef = SoftObjectPathTable.AddDefaulted_GetRef();
 			PathRef.SerializePath(MemAr);
 		}
@@ -1646,14 +1651,35 @@ void FAssetHeaderPatcherInner::PatchNameTable()
 
 bool FAssetHeaderPatcherInner::DoPatch(FSoftObjectPath& InOutSoft)
 {
-	FTopLevelAssetPath InOutTopLevelAssetPath = InOutSoft.GetAssetPath();
-	if (!DoPatch(InOutTopLevelAssetPath))
+	// FSoftObjectPaths are special in that while we may have an explict remapping 
+	// provided to the patcher, there may also be redirects stored in the global GRedirectCollector 
+	// reflecting on-disk object redirectors.
+	// Redirectors will be applied during FSoftObjectPath serialization so we need to account for
+	// that when patching so the nametable can be updated appropriately.
+
+	// Honour explictly remapped paths first. Even if this succeeds we need to handle redirectors
+	// because FSoftObjectPath serialization will always apply them.
+	FTopLevelAssetPath SrcTopLevelAssetPath = InOutSoft.GetAssetPath();
+	FTopLevelAssetPath InOutTopLevelAssetPath = SrcTopLevelAssetPath;
+	bool bPatched = DoPatch(InOutTopLevelAssetPath);
+	if (bPatched)
 	{
-		return false;
+		InOutSoft.SetPath(InOutTopLevelAssetPath, InOutSoft.GetSubPathString());
 	}
 
-	InOutSoft.SetPath(InOutTopLevelAssetPath, InOutSoft.GetSubPathString());
-	return true;
+// Only run in Editor builds as FSoftObjectPath only runs PreSavePath during SerializePath in editor builds
+#if WITH_EDITOR
+	if (InOutSoft.PreSavePath(nullptr))
+	{
+		FTopLevelAssetPath DstTopLevelAssetPath = InOutSoft.GetAssetPath();
+
+		bPatched |= RemapFName(InOutTopLevelAssetPath.GetAssetName(), DstTopLevelAssetPath.GetAssetName());
+		bPatched |= RemapFName(InOutTopLevelAssetPath.GetPackageName(), DstTopLevelAssetPath.GetPackageName());
+		checkf(bPatched, TEXT("Applied a AssetPathRedirection while patching '%s' -> '%s' but failed to update the nametable."), *SrcTopLevelAssetPath.ToString(), *DstTopLevelAssetPath.ToString());
+	}
+#endif
+
+	return bPatched;
 }
 
 bool FAssetHeaderPatcherInner::DoPatch(FObjectResource& InOutResource, bool bIsExport, FName& OutPatchedObjectName)
@@ -2841,6 +2867,8 @@ TEST_CASE_NAMED(FAssetHeaderPatcherTests, "AssetHeaderPatcher", "[AssetHeaderPat
 	const TCHAR*  DstPackageObjectPath		= TEXT("/DstMount/SomePath/DstPackage.DstPackage");
 	const TCHAR*  SrcMountName				= TEXT("/SrcSpecialMount/");
 	const TCHAR*  DstMountName				= TEXT("/DstSpecialMount/");
+	const FSoftObjectPath SoftObjectPathToRedirect(TEXT("/ToBeRedirectedMount/SomePath/ToBeRedirectedPackage.ToBeRedirectedPackage:Some.ToBeRedirectedPackage.Subobject"));
+	const FSoftObjectPath RedirectedSoftObjectPath(TEXT("/RedirectedMount/SomePath/RedirectedPackage.RedirectedPackage:Some.RedirectedPackage.Subobject"));
 	const FName   SrcPackagePathFName(SrcPackagePath);
 	const FName   DstPackagePathFName(DstPackagePath);
 	const FName   SrcAssetFName(TEXT("SrcPackage"));
@@ -2868,7 +2896,7 @@ TEST_CASE_NAMED(FAssetHeaderPatcherTests, "AssetHeaderPatcher", "[AssetHeaderPat
 	FAssetHeaderPatcherInner Patcher(DummySrcDstAsset, DummySrcDstAsset, StringReplacements, MountPointReplacementMap);
 
 	int32 OriginalNameTableCount = 0;
-	auto ResetPatcher = [&Patcher, &OriginalNameTableCount, SrcPackagePathFName, SrcAssetFName, SrcExportObjectFName, DummyImportPackagePathFName]()
+	auto ResetPatcher = [&Patcher, &OriginalNameTableCount, SrcPackagePathFName, SrcAssetFName, SrcExportObjectFName, &SoftObjectPathToRedirect, DummyImportPackagePathFName]()
 		{
 			// Reset NameTable
 			Patcher.NameTable.Empty();
@@ -2887,6 +2915,13 @@ TEST_CASE_NAMED(FAssetHeaderPatcherTests, "AssetHeaderPatcher", "[AssetHeaderPat
 			Patcher.NameTable.Add(FName(NAME_None));
 			Patcher.NameToIndexMap.Add(FName(DummyImportPackagePathFName).GetDisplayIndex(), Patcher.NameTable.Num());
 			Patcher.NameTable.Add(FName(DummyImportPackagePathFName));
+			// Redirected names (softpaths don't include the subpath in the name table
+			FName ToBeRedirectedPackageName = SoftObjectPathToRedirect.GetAssetPath().GetPackageName();
+			FName ToBeRedirectedAssetName = SoftObjectPathToRedirect.GetAssetPath().GetAssetName();
+			Patcher.NameToIndexMap.Add(ToBeRedirectedPackageName.GetDisplayIndex(), Patcher.NameTable.Num());
+			Patcher.NameTable.Add(ToBeRedirectedPackageName);
+			Patcher.NameToIndexMap.Add(ToBeRedirectedAssetName.GetDisplayIndex(), Patcher.NameTable.Num());
+			Patcher.NameTable.Add(ToBeRedirectedAssetName);
 
 			FObjectImport DummyImport;
 			DummyImport.ObjectName = Patcher.NameTable[Patcher.NameToIndexMap[DummyImportPackagePathFName.GetDisplayIndex()]];
@@ -3181,7 +3216,7 @@ TEST_CASE_NAMED(FAssetHeaderPatcherTests, "AssetHeaderPatcher", "[AssetHeaderPat
 		}
 	}
 
-	SECTION("DoPatch(FSoftObjectPath")
+	SECTION("DoPatch(FSoftObjectPath)")
 	{
 		{
 			ResetPatcher();
@@ -3211,9 +3246,56 @@ TEST_CASE_NAMED(FAssetHeaderPatcherTests, "AssetHeaderPatcher", "[AssetHeaderPat
 			CHECK_EQUALS(TEXT("SoftObject patching updates NameTable entry"), Patcher.NameTable[0], DstPackagePathFName);
 			CHECK_EQUALS(TEXT("SoftObject patching doesn't implicitly update the PackageFileSummary"), Patcher.Summary.NameCount, OriginalNameTableCount);
 		}
+
+#if WITH_EDITOR
+		SECTION("DoPatch(FSoftObjectPath)-WithRedirector")
+		{
+			ResetPatcher();
+
+			GRedirectCollector.AddAssetPathRedirection(SoftObjectPathToRedirect, RedirectedSoftObjectPath);
+			ON_SCOPE_EXIT
+			{
+				GRedirectCollector.RemoveAssetPathRedirection(SoftObjectPathToRedirect);
+			};
+
+			FSoftObjectPath Actual = SoftObjectPathToRedirect;
+			FSoftObjectPath Expected = RedirectedSoftObjectPath;
+			CHECK(Patcher.DoPatch(Actual));
+			CHECK_EQUALS(TEXT("SoftObjectPath with sub-object path patching"), Actual, Expected);
+			CHECK_EQUALS(TEXT("SoftObject patching doesn't implicitly update the NameTable"), Patcher.NameTable[0], SrcPackagePathFName);
+			CHECK_EQUALS(TEXT("SoftObject patching doesn't implicitly update the PackageFileSummary"), Patcher.Summary.NameCount, OriginalNameTableCount);
+			Patcher.PatchNameTable();
+			CHECK_EQUALS(TEXT("SoftObject patching doesn't implicitly update the PackageFileSummary"), Patcher.Summary.NameCount, OriginalNameTableCount);
+		}
+
+		SECTION("DoPatch(FSoftObjectPath)-WithRedirectorAndExplicitPatch")
+		{
+			ResetPatcher();
+
+			// Note SrcPackageObjectPath is NOT being redirected to DstPackageObjectPath (that mapping 
+			// provided to the patcher explicitly already)
+			FSoftObjectPath NameToRedirect(SrcPackageObjectPath);
+			GRedirectCollector.AddAssetPathRedirection(NameToRedirect, RedirectedSoftObjectPath);
+			ON_SCOPE_EXIT
+			{
+				GRedirectCollector.RemoveAssetPathRedirection(NameToRedirect);
+			};
+
+			// Even though we have a redirector, we also have a mapping specified to the patcher. 
+			// We give the patcher priority in such cases
+			FSoftObjectPath Actual = NameToRedirect;
+			FSoftObjectPath Expected(DstPackageObjectPath); 
+			CHECK(Patcher.DoPatch(Actual));
+			CHECK_EQUALS(TEXT("SoftObjectPath with sub-object path patching"), Actual, Expected);
+			CHECK_EQUALS(TEXT("SoftObject patching doesn't implicitly update the NameTable"), Patcher.NameTable[0], SrcPackagePathFName);
+			CHECK_EQUALS(TEXT("SoftObject patching doesn't implicitly update the PackageFileSummary"), Patcher.Summary.NameCount, OriginalNameTableCount);
+			Patcher.PatchNameTable();
+			CHECK_EQUALS(TEXT("SoftObject patching doesn't implicitly update the PackageFileSummary"), Patcher.Summary.NameCount, OriginalNameTableCount);
+		}
+#endif
 	}
 
-	SECTION("DoPatch(FTopLevelAssetPath")
+	SECTION("DoPatch(FTopLevelAssetPath)")
 	{
 		{
 			ResetPatcher();
@@ -3244,7 +3326,7 @@ TEST_CASE_NAMED(FAssetHeaderPatcherTests, "AssetHeaderPatcher", "[AssetHeaderPat
 		}
 	}
 
-	SECTION("DoPatch(FGatherableTextData")
+	SECTION("DoPatch(FGatherableTextData)")
 	{
 		{
 			ResetPatcher();
