Fixed AssetViewUtils::GetPackageLengthForCooking to consider Android_ASTCClient as a target platform name for cooking since it is larger than the existing LinuxArm64Server (thus considering worst case asset paths with 1-2 chars less than it should)

diff --git a/Engine/Source/Developer/AssetTools/Private/AssetViewUtils.cpp b/Engine/Source/Developer/AssetTools/Private/AssetViewUtils.cpp
--- a/Engine/Source/Developer/AssetTools/Private/AssetViewUtils.cpp
+++ b/Engine/Source/Developer/AssetTools/Private/AssetViewUtils.cpp
@@ -1570,8 +1570,8 @@ int32 AssetViewUtils::GetPackageLengthForCooking(const FString& PackageName, boo
 	const bool bIsEngineAsset = IsEngineFolder(PackageName) || (PluginContainingAsset && PluginContainingAsset->GetLoadedFrom() == EPluginLoadedFrom::Engine);
 	const bool bIsProjectAsset = !bIsEngineAsset;
 
-	// We use "LinuxArm64Server" below as it's probably the longest platform name, so will also prove that any shorter platform names will validate correctly
-	const FString CookSubPath = FPaths::Combine(TEXT("Saved"), TEXT("Cooked"), TEXT("LinuxArm64Server"), bIsEngineAsset ? TEXT("Engine") : GameName, TEXT("")); // Trailing empty entry ensures a trailing slash
+	// We use "Android_ASTCClient" below as it's probably the longest platform name, so will also prove that any shorter platform names will validate correctly
+	const FString CookSubPath = FPaths::Combine(TEXT("Saved"), TEXT("Cooked"), TEXT("Android_ASTCClient"), bIsEngineAsset ? TEXT("Engine") : GameName, TEXT("")); // Trailing empty entry ensures a trailing slash
 	const FString AbsolutePathToAsset = FPaths::ConvertRelativePathToFull(RelativePathToAsset);
 
 	FString AbsoluteTargetPath = FPaths::ConvertRelativePathToFull(bIsEngineAsset ? FPaths::EngineDir() : FPaths::ProjectDir());
