[AssetHeaderPatcher]
Fix issue where if we cannot find the packagepath in the nametable, we would assume we could use the original package path when forming FTopLevelAssetPaths for AssetRegistry ObjectPath data, since it isn't always a full path.

We now only fill in the implicit package path if we did find the original package name in the nametable so we don't inadvertently create mappings of objectpaths that don't exist in the nametable to begin with.

#rnx

diff --git a/Engine/Source/Developer/AssetTools/Private/AssetHeaderPatcher.cpp b/Engine/Source/Developer/AssetTools/Private/AssetHeaderPatcher.cpp
--- a/Engine/Source/Developer/AssetTools/Private/AssetHeaderPatcher.cpp
+++ b/Engine/Source/Developer/AssetTools/Private/AssetHeaderPatcher.cpp
@@ -794,6 +794,7 @@ public:
 		, StringMountPointReplacements(InStringMountPointReplacements)
 		, DstArchive(InDstArchive)
 		, bPatchPrimaryAssetTag(false)
+		, bIsPackagePathInNametable(false)
 		, bIsNonOneFilePerActorPackage(false)
 	{
 		for (auto TagToIgnore : TagsToIgnore)
@@ -851,6 +852,7 @@ public:
 	FName OriginalNonOneFilePerActorPackagePath; // e.g. "/MountName/MountName"
 	FString OriginalPrimaryAssetName;			 // e.g. "MountName"
 	bool bPatchPrimaryAssetTag;
+	bool bIsPackagePathInNametable;
 	bool bIsNonOneFilePerActorPackage;
 
 	// NameTable Members
@@ -1879,6 +1881,7 @@ void FAssetHeaderPatcherInner::PatchHeader_PatchSections()
 		// extend this special case of NameTable patching to all packages, OFPA or not.
 		if (NameToIndexMap.Find(OriginalPackagePath.GetDisplayIndex()))
 		{
+			bIsPackagePathInNametable = true;
 			RemapFName(OriginalPackagePath, DstPackageName.PackageName);
 		}
 	}
@@ -2012,7 +2015,9 @@ void FAssetHeaderPatcherInner::PatchHeader_PatchSections()
 					ObjData.ObjectData.ObjectPath = SrdDstObjectPath.ToString();
 				}
 			}
-			else
+			// Only use OriginalPackagePath if it was in the nametable to begin with, 
+			// since if it wasn't we can't have a remapping for the ObjectPath
+			else if(bIsPackagePathInNametable)
 			{
 				FTopLevelAssetPath SrcDstTopLevelAssetPath(OriginalPackagePath, FName(ObjData.ObjectData.ObjectPath));
 				SrcObjectPath.SetPath(SrcDstTopLevelAssetPath, SrcObjectPath.GetSubPathString());
