[DuplicatePluginCommandlet]
Add a file version check before using AssetHeaderPatching and if any incompatible files are found, fallback to using one of our other duplication mechanisms

#rnx

diff --git a/Engine/Source/Developer/AssetTools/Internal/AssetHeaderPatcher.h b/Engine/Source/Developer/AssetTools/Internal/AssetHeaderPatcher.h
--- a/Engine/Source/Developer/AssetTools/Internal/AssetHeaderPatcher.h
+++ b/Engine/Source/Developer/AssetTools/Internal/AssetHeaderPatcher.h
@@ -15,6 +15,8 @@ DECLARE_DELEGATE_TwoParams(FAssetHeaderPatcherCompletionDelegate, const FString&
 
 UE_INTERNAL struct ASSETTOOLS_API FAssetHeaderPatcher
 {
+	static const EUnrealEngineObjectUE5Version MinimumSupportedUE5FileVersion = EUnrealEngineObjectUE5Version::ADD_SOFTOBJECTPATH_LIST;
+
 	enum class EResult
 	{
 		NotStarted,
