Fix asset categories:
- FX no longer uses two different bits
- Editor Utilities no longer uses two different bits
- Artificial Intelligence no longer uses two different bits

diff --git a/Engine/Source/Developer/AssetTools/Private/AssetTools.cpp b/Engine/Source/Developer/AssetTools/Private/AssetTools.cpp
--- a/Engine/Source/Developer/AssetTools/Private/AssetTools.cpp
+++ b/Engine/Source/Developer/AssetTools/Private/AssetTools.cpp
@@ -1407,6 +1407,7 @@ UAssetToolsImpl::UAssetToolsImpl(const FObjectInitializer& ObjectInitializer)
 	AllocatedCategoryBits.Add(TEXT("Media"), FAdvancedAssetCategory(EAssetTypeCategories::Media, LOCTEXT("MediaAssetCategory", "Media")));
 	AllocatedCategoryBits.Add(TEXT("Texture"), FAdvancedAssetCategory(EAssetTypeCategories::Textures, LOCTEXT("TextureAssetCategory", "Texture")));
 	AllocatedCategoryBits.Add(TEXT("World"), FAdvancedAssetCategory(EAssetTypeCategories::World, LOCTEXT("WorldAssetCategory", "World")));
+	AllocatedCategoryBits.Add(TEXT("FX"), FAdvancedAssetCategory(EAssetTypeCategories::FX, LOCTEXT("FXAssetCategory", "FX")));
 
 	// Register the built-in asset type actions
 	RegisterAssetTypeActions(MakeShareable(new FAssetTypeActions_AnimBoneCompressionSettings));
@@ -1629,6 +1630,8 @@ EAssetTypeCategories::Type UAssetToolsImpl::RegisterAdvancedAssetCategory(FName
 			Result = (EAssetTypeCategories::Type)NextUserCategoryBit;
 			AllocatedCategoryBits.Add(CategoryKey, FAdvancedAssetCategory(Result, CategoryDisplayName));
 
+			UE_LOG(LogAssetTools, Verbose, TEXT("RegisterAssetTypeCategory(\"%s\", \"%s\") used up bit 0x%x, Offset: %d"), *CategoryKey.ToString(), *CategoryDisplayName.ToString(), NextUserCategoryBit, FMath::FloorLog2(Result));
+
 			// Advance to the next bit, or store that we're out
 			if (NextUserCategoryBit == EAssetTypeCategories::LastUser)
 			{
diff --git a/Engine/Source/Runtime/AIModule/Private/AIModule.cpp b/Engine/Source/Runtime/AIModule/Private/AIModule.cpp
--- a/Engine/Source/Runtime/AIModule/Private/AIModule.cpp
+++ b/Engine/Source/Runtime/AIModule/Private/AIModule.cpp
@@ -66,7 +66,7 @@ void FAIModule::StartupModule()
 		IAssetTools& AssetTools = FModuleManager::LoadModuleChecked<FAssetToolsModule>("AssetTools").Get();
 
 		// register AI category so that AI assets can register to it
-		AIAssetCategoryBit = AssetTools.RegisterAdvancedAssetCategory(FName(TEXT("AI")), LOCTEXT("AIAssetCategory", "Artificial Intelligence"));
+		AIAssetCategoryBit = AssetTools.RegisterAdvancedAssetCategory(FName(TEXT("Artificial Intelligence")), LOCTEXT("AIAssetCategory", "Artificial Intelligence"));
 	}
 #endif // WITH_EDITOR 
 
