Improved OpenLocalizedVariantsDialog with multiple choices: Include/Exclude/Cancel (instead of just Close)

#rnx

diff --git a/Engine/Source/Developer/AssetTools/Private/AssetRenameManager.cpp b/Engine/Source/Developer/AssetTools/Private/AssetRenameManager.cpp
--- a/Engine/Source/Developer/AssetTools/Private/AssetRenameManager.cpp
+++ b/Engine/Source/Developer/AssetTools/Private/AssetRenameManager.cpp
@@ -359,9 +359,20 @@ bool FAssetRenameManager::RenameAssetsAndVariants(const TArray<FAssetRenameData>
 			}
 			if (!NewAssetsAdded.IsEmpty())
 			{
-				LocalizedAssetTools->OpenLocalizedVariantsListMessageDialog(LOCTEXT("AddedLocalizedVariantsHeader", "Renaming localized variants too"),
-					LOCTEXT("AddedLocalizedVariantsMessage", "The following localized variants (or source assets) are also going to be renamed alongside the selected assets."),
-					NewAssetsAdded);
+				ELocalizedVariantsInclusion InclusionResult = LocalizedAssetTools->OpenIncludeLocalizedVariantsListDialog(NewAssetsAdded);
+				switch (InclusionResult)
+				{
+				case ELocalizedVariantsInclusion::Exclude:
+					AssetsAndVariants = InAssetsRenameData;
+					break;
+				case ELocalizedVariantsInclusion::Cancel:
+					AssetsAndVariants = InAssetsRenameData;
+					return false;
+				case ELocalizedVariantsInclusion::Include:
+				default:
+					// The inclusion is already done at this point. Only the exclude/cancel options need to do some revert work.
+					break;
+				}
 			}
 		}
 	}
@@ -2020,7 +2031,7 @@ int32 FAssetRenameManager::ReportFailures(const TArray<FAssetRenameDataWithRefer
 	{
 		if (bWithDialog)
 		{
-			SFileListReportDialog::OpenDialog(LOCTEXT("FailedRenamesDialog", "Failed Renames"),
+			SFileListReportDialog::OpenListDialog(LOCTEXT("FailedRenamesDialog", "Failed Renames"),
 				LOCTEXT("RenameFailureTitle", "The following assets could not be renamed."),
 				FailedRenames);
 		}
diff --git a/Engine/Source/Developer/AssetTools/Private/LocalizedAssetTools.cpp b/Engine/Source/Developer/AssetTools/Private/LocalizedAssetTools.cpp
--- a/Engine/Source/Developer/AssetTools/Private/LocalizedAssetTools.cpp
+++ b/Engine/Source/Developer/AssetTools/Private/LocalizedAssetTools.cpp
@@ -11,14 +11,145 @@
 #include "ISourceControlModule.h"
 #include "ISourceControlProvider.h"
 #include "Logging/LogMacros.h"
+#include "Misc/App.h"
 #include "Misc/MessageDialog.h"
 #include "SFileListReportDialog.h"
 #include "SourceControlHelpers.h"
 #include "SourceControlPreferences.h"
 #include "UObject/UObjectGlobals.h"
+#include "Widgets/Input/SButton.h"
 
 #define LOCTEXT_NAMESPACE "LocalizedAssetTools"
 
+class SIncludeLocalizedVariantsDialog : public SFileListReportDialog
+{
+private:
+	bool bAllowOperationCanceling = true;
+	ELocalizedVariantsInclusion RecommendedBehavior = ELocalizedVariantsInclusion::Include;
+	ELocalizedVariantsInclusion Result = ELocalizedVariantsInclusion::Include;
+
+public:
+	static ELocalizedVariantsInclusion OpenIncludeListDialog(const FText& InTitle, const FText& InHeader, const TArray<FText>& InFiles, ELocalizedVariantsInclusion RecommendedBehavior = ELocalizedVariantsInclusion::Include, bool bAllowOperationCanceling = true, ELocalizedVariantsInclusion UnattendedDefaultBehavior = ELocalizedVariantsInclusion::Exclude)
+	{
+		// Make sure that if bAllowOperationCanceling is false then the default behavior should never be "Cancel".
+		ensureAlways(bAllowOperationCanceling || (!bAllowOperationCanceling && RecommendedBehavior != ELocalizedVariantsInclusion::Cancel));
+
+		if (FApp::IsUnattended() || GIsRunningUnattendedScript)
+		{
+			return UnattendedDefaultBehavior;
+		}
+
+		TSharedRef<SIncludeLocalizedVariantsDialog> FileListReportDialogRef = SNew(SIncludeLocalizedVariantsDialog).Header(InHeader).Files(InFiles);
+
+		FileListReportDialogRef->bAllowOperationCanceling = bAllowOperationCanceling;
+		FileListReportDialogRef->RecommendedBehavior = RecommendedBehavior;
+		FileListReportDialogRef->Result = RecommendedBehavior;
+
+		FileListReportDialogRef->bOpenAsModal = true;
+		FileListReportDialogRef->bAllowTitleBarX = bAllowOperationCanceling;
+		FileListReportDialogRef->Title = InTitle;
+		CreateWindow(FileListReportDialogRef);
+
+		// Modal window is closed so the result should be accessible now
+
+		return FileListReportDialogRef->Result;
+	}
+
+protected:
+	FReply OnIncludeSelected()
+	{
+		Result = ELocalizedVariantsInclusion::Include;
+		return CloseWindow();
+	}
+
+	FReply OnExcludeSelected()
+	{
+		Result = ELocalizedVariantsInclusion::Exclude;
+		return CloseWindow();
+	}
+
+	FReply OnCancelSelected()
+	{
+		Result = ELocalizedVariantsInclusion::Cancel;
+		return CloseWindow();
+	}
+
+	virtual void OnClosedWithTitleBarX(const TSharedRef<SWindow>& Window)
+	{
+		if (bAllowOperationCanceling)
+		{
+			Result = ELocalizedVariantsInclusion::Cancel;
+		}
+
+		// Otherwise, Result is already defined, let's use that value
+	}
+
+	virtual TSharedRef<SHorizontalBox> ConstructButtons(const FArguments& InArgs) override
+	{
+		TSharedRef<SHorizontalBox> HorizontalBox = SNew(SHorizontalBox)
+			+ SHorizontalBox::Slot()
+			.AutoWidth()
+			.Padding(4, 0)
+			.HAlign(HAlign_Right)
+			[
+				SNew(SButton)
+					.OnClicked(this, &SIncludeLocalizedVariantsDialog::OnIncludeSelected)
+					.Text(LOCTEXT("IncludeLocalizedVariantsDialogIncludeButtonText", "Include"))
+					.ButtonStyle(FAppStyle::Get(), RecommendedBehavior == ELocalizedVariantsInclusion::Include ? "FlatButton.Primary" : "FlatButton.Default")
+			]
+			+ SHorizontalBox::Slot()
+			.AutoWidth()
+			.Padding(4, 0)
+			.HAlign(HAlign_Right)
+			[
+				SNew(SButton)
+					.OnClicked(this, &SIncludeLocalizedVariantsDialog::OnExcludeSelected)
+					.Text(LOCTEXT("IncludeLocalizedVariantsDialogExcludeButtonText", "Exclude"))
+					.ButtonStyle(FAppStyle::Get(), RecommendedBehavior == ELocalizedVariantsInclusion::Exclude ? "FlatButton.Primary" : "FlatButton.Default")
+			];
+
+		if (bAllowOperationCanceling)
+		{
+			HorizontalBox->AddSlot()
+				.AutoWidth()
+				.Padding(4, 0)
+				.HAlign(HAlign_Right)
+				[
+					SNew(SButton)
+						.OnClicked(this, &SIncludeLocalizedVariantsDialog::OnCancelSelected)
+						.Text(LOCTEXT("IncludeLocalizedVariantsDialogCancelButtonText", "Cancel"))
+						.ButtonStyle(FAppStyle::Get(), RecommendedBehavior == ELocalizedVariantsInclusion::Cancel ? "FlatButton.Primary" : "FlatButton.Default")
+				];
+		}
+
+		return HorizontalBox;
+	}
+
+	virtual FReply OnKeyDown(const FGeometry& MyGeometry, const FKeyEvent& InKeyEvent) override
+	{
+		if (InKeyEvent.GetKey() == EKeys::Enter)
+		{
+			switch (RecommendedBehavior)
+			{
+			case ELocalizedVariantsInclusion::Include:
+				return OnIncludeSelected();
+			case ELocalizedVariantsInclusion::Exclude:
+				return OnExcludeSelected();
+			case ELocalizedVariantsInclusion::Cancel:
+			default:
+				return OnCancelSelected();
+				break;
+			}
+		}
+		else if(InKeyEvent.GetKey() == EKeys::Escape)
+		{
+			return OnCancelSelected();
+		}
+
+		return SFileListReportDialog::OnKeyDown(MyGeometry, InKeyEvent);
+	}
+};
+
 FLocalizedAssetTools::FLocalizedAssetTools() : 
 	RevisionControlIsNotAvailableWarningText(LOCTEXT("RevisionControlIsRequiredToChangeLocalizableAssets", "Revision Control is required to move/rename/delete localizable assets for this project and it is currently not accessible."))
 	, FilesNeedToBeOnDiskWarningText(LOCTEXT("FilesToSyncDialogTitle", "Files in Revision Control need to be on disk"))
@@ -210,8 +341,8 @@ ELocalizedAssetsResult FLocalizedAssetTools::GetLocalizedVariants(const TArray<F
 void FLocalizedAssetTools::OpenRevisionControlRequiredDialog() const
 {
 	FText WarningText = RevisionControlIsNotAvailableWarningText;
-	const FText AvoidWarningText = LOCTEXT("HowToFixRevisionControlIsRequiredToChangeLocalizableAssets", "If you want to disable this project option, it is located under:\n\tProject Settings/\n\tEditor/\n\tRevision Control/\n\tRequires Revision Control To Rename Localizable Assets\n\nThis option is there to prevent breaking paths between a source asset and its localized variants if they are not on disk.");
-	FMessageDialog::Open(EAppMsgType::Ok, WarningText.Format(LOCTEXT("RevisionControlIsRequiredToChangeLocalizableAssetsDialog", "{0}\n\n{1}"), WarningText, AvoidWarningText));
+	const FText AvoidWarningText = LOCTEXT("HowToFixRevisionControlIsRequiredToManageLocalizableAssets", "If you want to disable this project option, it is located under:\n\tProject Settings/\n\tEditor/\n\tRevision Control/\n\tRequires Revision Control To Manage Localizable Assets\n\nThis option is there to prevent breaking paths between a source asset and its localized variants if they are not on disk.");
+	FMessageDialog::Open(EAppMsgType::Ok, WarningText.Format(LOCTEXT("RevisionControlIsRequiredToManageLocalizableAssetsDialog", "{0}\n\n{1}"), WarningText, AvoidWarningText));
 }
 
 void FLocalizedAssetTools::OpenFilesInRevisionControlRequiredDialog(const TArray<FText>& FileList) const
@@ -223,7 +354,17 @@ void FLocalizedAssetTools::OpenFilesInRevisionControlRequiredDialog(const TArray
 
 void FLocalizedAssetTools::OpenLocalizedVariantsListMessageDialog(const FText& Header, const FText& Message, const TArray<FText>& FileList) const
 {
-	SFileListReportDialog::OpenDialog(Header, Message, FileList, true);
+	SFileListReportDialog::OpenListDialog(Header, Message, FileList, true);
+}
+
+ELocalizedVariantsInclusion FLocalizedAssetTools::OpenIncludeLocalizedVariantsListDialog(const TArray<FText>& FileList, ELocalizedVariantsInclusion RecommendedBehavior /*= ELocalizedVariantsInclusion::Include*/, bool bAllowOperationCanceling /*= true*/, ELocalizedVariantsInclusion UnattendedDefaultBehavior /*= ELocalizedVariantsInclusion::Exclude*/) const
+{
+	return SIncludeLocalizedVariantsDialog::OpenIncludeListDialog(LOCTEXT("IncludeLocalizedVariantsDialogTitle", "Include Localized Variants"),
+		LOCTEXT("IncludeLocalizedVariantsDialogHeader", "The current operation could also apply to the following localized variants (or source asset). Do you want to include them in the current operation ?"),
+		FileList,
+		RecommendedBehavior,
+		bAllowOperationCanceling,
+		UnattendedDefaultBehavior);
 }
 
 const FText& FLocalizedAssetTools::GetRevisionControlIsNotAvailableWarningText() const
diff --git a/Engine/Source/Developer/AssetTools/Private/LocalizedAssetTools.h b/Engine/Source/Developer/AssetTools/Private/LocalizedAssetTools.h
--- a/Engine/Source/Developer/AssetTools/Private/LocalizedAssetTools.h
+++ b/Engine/Source/Developer/AssetTools/Private/LocalizedAssetTools.h
@@ -29,6 +29,8 @@ public:
 
 	virtual void OpenLocalizedVariantsListMessageDialog(const FText& Header, const FText& Message, const TArray<FText>& FileList) const override;
 
+	virtual ELocalizedVariantsInclusion OpenIncludeLocalizedVariantsListDialog( const TArray<FText>& FileList, ELocalizedVariantsInclusion RecommendedBehavior = ELocalizedVariantsInclusion::Include, bool bAllowOperationCanceling = true, ELocalizedVariantsInclusion UnattendedDefaultBehavior = ELocalizedVariantsInclusion::Exclude) const override;
+
 	virtual const FText& GetRevisionControlIsNotAvailableWarningText() const override;
 
 	virtual const FText& GetFilesNeedToBeOnDiskWarningText() const override;
diff --git a/Engine/Source/Developer/AssetTools/Private/SFileListReportDialog.cpp b/Engine/Source/Developer/AssetTools/Private/SFileListReportDialog.cpp
--- a/Engine/Source/Developer/AssetTools/Private/SFileListReportDialog.cpp
+++ b/Engine/Source/Developer/AssetTools/Private/SFileListReportDialog.cpp
@@ -4,6 +4,7 @@
 
 #include "Framework/Application/SlateApplication.h"
 #include "Interfaces/IMainFrameModule.h"
+#include "Misc/App.h"
 #include "Widgets/Input/SButton.h"
 #include "Widgets/Layout/SBorder.h"
 #include "Widgets/SWindow.h"
@@ -34,6 +35,7 @@ void SFileListReportDialog::Construct(const FArguments& InArgs)
 						.AutoHeight()
 						[
 							SNew(STextBlock).Text(InArgs._Header)
+								.AutoWrapText(true)
 						]
 
 						// Files To Sync list
@@ -51,36 +53,94 @@ void SFileListReportDialog::Construct(const FArguments& InArgs)
 								]
 						]
 
-						// Close button
+						// Buttons
 						+ SVerticalBox::Slot()
 						.AutoHeight()
 						.Padding(0, 4)
 						.HAlign(HAlign_Right)
 						[
-							SNew(SButton)
-								.OnClicked(this, &SFileListReportDialog::CloseClicked)
-								.Text(LOCTEXT("WindowCloseButton", "Close"))
+							ConstructButtons(InArgs)
 						]
 				]
 		];
 }
 
-void SFileListReportDialog::OpenDialog(const FText& InTitle, const FText& InHeader, const TArray<FText>& InFiles, bool bOpenAsModal /*= false*/)
+void SFileListReportDialog::OpenListDialog(const FText& InTitle, const FText& InHeader, const TArray<FText>& InFiles, bool bOpenAsModal /*= false*/)
+{
+	if (FApp::IsUnattended() || GIsRunningUnattendedScript)
+	{
+		return;
+	}
+
+	TSharedRef<SFileListReportDialog> FileListReportDialogRef = SNew(SFileListReportDialog).Header(InHeader).Files(InFiles);
+	FileListReportDialogRef->bOpenAsModal = bOpenAsModal;
+	FileListReportDialogRef->bAllowTitleBarX = true;
+	FileListReportDialogRef->Title = InTitle;
+	CreateWindow(FileListReportDialogRef);
+}
+
+TSharedRef<ITableRow> SFileListReportDialog::MakeListViewWidget(TSharedRef<FText> Item, const TSharedRef<STableViewBase>& OwnerTable)
+{
+	return
+		SNew(STableRow< TSharedRef<FText> >, OwnerTable)
+		[
+			SNew(STextBlock).Text(Item.Get())
+		];
+}
+
+TSharedRef<SHorizontalBox> SFileListReportDialog::ConstructButtons(const FArguments& InArgs)
+{
+	return SNew(SHorizontalBox)
+		+ SHorizontalBox::Slot()
+		.AutoWidth()
+		.HAlign(HAlign_Right)
+		[
+			SNew(SButton)
+				.OnClicked(this, &SFileListReportDialog::CloseWindow)
+				.Text(LOCTEXT("WindowCloseButton", "Close"))
+		];
+}
+
+FReply SFileListReportDialog::OnKeyDown(const FGeometry& MyGeometry, const FKeyEvent& InKeyEvent)
+{
+	if (InKeyEvent.GetKey() == EKeys::Enter || InKeyEvent.GetKey() == EKeys::Escape)
+	{
+		CloseWindow();
+		return FReply::Handled();
+	}
+
+	return SCompoundWidget::OnKeyDown(MyGeometry, InKeyEvent);
+}
+
+bool SFileListReportDialog::SupportsKeyboardFocus() const
+{
+	return true;
+}
+
+void SFileListReportDialog::CreateWindow(TSharedRef<SFileListReportDialog> InFileListReportDialogRef)
 {
 	TSharedRef<SWindow> FileListReportWindow = SNew(SWindow)
-		.Title(InTitle)
+		.Title(InFileListReportDialogRef->Title)
 		.ClientSize(FVector2D(800, 400))
 		.SupportsMaximize(false)
 		.SupportsMinimize(false)
+		.HasCloseButton(InFileListReportDialogRef->bAllowTitleBarX)
 		[
-			SNew(SFileListReportDialog).Header(InHeader).Files(InFiles)
+			InFileListReportDialogRef
 		];
 
+	FileListReportWindow->GetOnWindowActivatedEvent().AddLambda([InFileListReportDialogRef]()
+		{
+			FSlateApplication::Get().SetKeyboardFocus(InFileListReportDialogRef);
+		});
+
+	FileListReportWindow->SetOnWindowClosed(FOnWindowClosed::CreateRaw(&InFileListReportDialogRef.Get(), &SFileListReportDialog::OnWindowClosed));
+
 	IMainFrameModule& MainFrameModule = FModuleManager::LoadModuleChecked<IMainFrameModule>(TEXT("MainFrame"));
-	
+
 	if (MainFrameModule.GetParentWindow().IsValid())
 	{
-		if (bOpenAsModal)
+		if (InFileListReportDialogRef->bOpenAsModal)
 		{
 			FSlateApplication::Get().AddModalWindow(FileListReportWindow, MainFrameModule.GetParentWindow().ToSharedRef());
 		}
@@ -91,21 +151,21 @@ void SFileListReportDialog::OpenDialog(const FText& InTitle, const FText& InHead
 	}
 	else
 	{
-		FSlateApplication::Get().AddWindow(FileListReportWindow);
+		if (InFileListReportDialogRef->bOpenAsModal)
+		{
+			FSlateApplication::Get().AddModalWindow(FileListReportWindow, nullptr);
+		}
+		else
+		{
+			FSlateApplication::Get().AddWindow(FileListReportWindow);
+		};
 	}
 }
 
-TSharedRef<ITableRow> SFileListReportDialog::MakeListViewWidget(TSharedRef<FText> Item, const TSharedRef<STableViewBase>& OwnerTable)
+FReply SFileListReportDialog::CloseWindow()
 {
-	return
-		SNew(STableRow< TSharedRef<FText> >, OwnerTable)
-		[
-			SNew(STextBlock).Text(Item.Get())
-		];
-}
+	bClosingWithoutTitleBarX = true;
 
-FReply SFileListReportDialog::CloseClicked()
-{
 	TSharedPtr<SWindow> Window = FSlateApplication::Get().FindWidgetWindow(AsShared());
 
 	if (Window.IsValid())
@@ -116,4 +176,22 @@ FReply SFileListReportDialog::CloseClicked()
 	return FReply::Handled();
 }
 
+void SFileListReportDialog::OnWindowClosed(const TSharedRef<SWindow>& Window)
+{
+	if (!bClosingWithoutTitleBarX)
+	{
+		OnClosedWithTitleBarX(Window);
+	}
+}
+
+void SFileListReportDialog::SetModal(bool bInIsModal)
+{
+	bOpenAsModal = bInIsModal;
+}
+
+void SFileListReportDialog::SetAllowTitleBarX(bool bInAllowTitleBarX)
+{
+	bAllowTitleBarX = bInAllowTitleBarX;
+}
+
 #undef LOCTEXT_NAMESPACE
diff --git a/Engine/Source/Developer/AssetTools/Private/SFileListReportDialog.h b/Engine/Source/Developer/AssetTools/Private/SFileListReportDialog.h
--- a/Engine/Source/Developer/AssetTools/Private/SFileListReportDialog.h
+++ b/Engine/Source/Developer/AssetTools/Private/SFileListReportDialog.h
@@ -6,6 +6,7 @@
 #include "Widgets/SCompoundWidget.h"
 
 class ITableRow;
+class SHorizontalBox;
 class STableViewBase;
 
 class SFileListReportDialog : public SCompoundWidget
@@ -20,14 +21,32 @@ public:
 
 	void Construct(const FArguments& InArgs);
 
-	static void OpenDialog(const FText& InTitle, const FText& InHeader, const TArray<FText>& InFiles, bool bOpenAsModal = false);
+	static void OpenListDialog(const FText& InTitle, const FText& InHeader, const TArray<FText>& InFiles, bool bOpenAsModal = false);
 
-private:
+protected:
 	TSharedRef<ITableRow> MakeListViewWidget(TSharedRef<FText> Item, const TSharedRef<STableViewBase>& OwnerTable);
+	virtual TSharedRef<SHorizontalBox> ConstructButtons(const FArguments& InArgs);
+
+	virtual FReply OnKeyDown(const FGeometry& MyGeometry, const FKeyEvent& InKeyEvent) override;
+	virtual bool SupportsKeyboardFocus() const override;
+
+	static void CreateWindow(TSharedRef<SFileListReportDialog> InFileListReportDialogRef);
 
-	FReply CloseClicked();
+	FReply CloseWindow();
+	void OnWindowClosed(const TSharedRef<SWindow>& Window);
+	virtual void OnClosedWithTitleBarX(const TSharedRef<SWindow>& Window) {}
+
+	void SetModal(bool bIsModal);
+	void SetAllowTitleBarX(bool bInAllowTitleBarX);
+
+protected:
+	bool bOpenAsModal = false;
+	bool bAllowTitleBarX = true;
+	FText Title;
 
 private:
 	FText Header;
 	TArray< TSharedRef<FText> > Files;
+
+	bool bClosingWithoutTitleBarX = false;
 };
diff --git a/Engine/Source/Developer/AssetTools/Public/ILocalizedAssetTools.h b/Engine/Source/Developer/AssetTools/Public/ILocalizedAssetTools.h
--- a/Engine/Source/Developer/AssetTools/Public/ILocalizedAssetTools.h
+++ b/Engine/Source/Developer/AssetTools/Public/ILocalizedAssetTools.h
@@ -30,6 +30,16 @@ enum class ELocalizedAssetsResult : uint8
 	RevisionControlNotAvailable,
 };
 
+enum class ELocalizedVariantsInclusion : uint8
+{
+	/** Include localized variants (and possibly source assets) to an operation */
+	Include,
+	/** Exclude localized variants (or related source asset) from an operation */
+	Exclude,
+	/** Cancel the current operation if the user do not want to choose between Including or Excluding localized variants. */
+	Cancel
+};
+
 class ILocalizedAssetTools
 {
 public:
@@ -53,6 +63,9 @@ public:
 	/** Open a dialog with a file list (localized variants) with a custom message */
 	virtual void OpenLocalizedVariantsListMessageDialog(const FText& Header, const FText& Message, const TArray<FText>& FileList) const = 0;
 
+	/** Open a dialog to ask to include/exclude localized variants from an operation */
+	virtual ELocalizedVariantsInclusion OpenIncludeLocalizedVariantsListDialog(const TArray<FText>& FileList, ELocalizedVariantsInclusion RecommendedBehavior = ELocalizedVariantsInclusion::Include, bool bAllowOperationCanceling = true, ELocalizedVariantsInclusion UnattendedDefaultBehavior = ELocalizedVariantsInclusion::Exclude) const = 0;
+
 	/** Get shared text to warn that some files need to be on disk (not only in Revision Control) */
 	virtual const FText& GetFilesNeedToBeOnDiskWarningText() const = 0;
 
