[AssetHeaderPatcher]

Do not include ObjectNames in the StringReplacement mappings as the bare, non-fully-qualified names are too generic and are likely to lead to incorrect replacements in string data.

e.g. Remapping "/SomeMount/Name.uasset" -> "/NewMount/Name.uasset" would create a mapping for "Name" -> "Name" which later when looked up via FName will cause lowercase instances of "name" to be promoted to "Name" which might be incorrect to do (the two "names" may not even be contextually the same 'name').

We do not make string replacements case sensitive to alleviate the above case swapping because string data often has legitimate paths that are not the correct case, however if we can guarantee that later it would be worth doing.

#rnx

diff --git a/Engine/Source/Developer/AssetTools/Private/AssetHeaderPatcher.cpp b/Engine/Source/Developer/AssetTools/Private/AssetHeaderPatcher.cpp
--- a/Engine/Source/Developer/AssetTools/Private/AssetHeaderPatcher.cpp
+++ b/Engine/Source/Developer/AssetTools/Private/AssetHeaderPatcher.cpp
@@ -498,7 +498,7 @@ void FAssetHeaderPatcher::FContext::GenerateAdditionalRemappings()
 		const FCoreRedirectObjectName& SrcName = Redirect.OldName;
 		const FCoreRedirectObjectName& DstName = Redirect.NewName;
 
-		StringReplacements.Add(SrcName.ObjectName.ToString(), DstName.ObjectName.ToString());
+		// Do not include Src->Dst ObjectName mappings since it's too likely that we will incorrectly rename when dealing with string data
 		StringReplacements.Add(SrcName.PackageName.ToString(), DstName.PackageName.ToString());
 		StringReplacements.Add(SrcName.ToString(), DstName.ToString());
 
@@ -3057,6 +3057,14 @@ TEST_CASE_NAMED(FAssetHeaderPatcherTests, "AssetHeaderPatcher", "[AssetHeaderPat
 				CHECK(!Patcher.DoPatch(Actual));
 				CHECK_EQUALS(TEXT("Patch string with no direct match"), Actual, Expected);
 			}
+
+			// Do not replace strings that happen to overlap with ObjectNames that are remapped
+			{
+				FString Actual(TEXT(R"(SrcPackage)"));
+				const FString Expected(TEXT(R"(SrcPackage)"));
+				CHECK(!Patcher.DoPatch(Actual));
+				CHECK_EQUALS(TEXT("Do not remap a string that matches a non-fully-qualified ObjectName"), Actual, Expected);
+			}
 		}
 
 		SECTION("Sub-Object Paths")
@@ -3110,6 +3118,14 @@ TEST_CASE_NAMED(FAssetHeaderPatcherTests, "AssetHeaderPatcher", "[AssetHeaderPat
 					CHECK(Patcher.DoPatch(Actual));
 					CHECK_EQUALS(TEXT("Patch default blueprint generated class object path with quotes"), Actual, Expected);
 				}
+
+				// Do not replace strings that happen to overlap with ObjectNames that are remapped
+				{
+					FString Actual(TEXT(R"('SrcPackage')"));
+					const FString Expected(TEXT(R"('SrcPackage')"));
+					CHECK(!Patcher.DoPatch(Actual));
+					CHECK_EQUALS(TEXT("Do not remap a string that matches a non-fully-qualified ObjectName"), Actual, Expected);
+				}
 			}
 
 			SECTION("Double Quote")
@@ -3141,6 +3157,14 @@ TEST_CASE_NAMED(FAssetHeaderPatcherTests, "AssetHeaderPatcher", "[AssetHeaderPat
 					CHECK(Patcher.DoPatch(Actual));
 					CHECK_EQUALS(TEXT("Patch default blueprint generated class object path with quotes"), Actual, Expected);
 				}
+
+				// Do not replace strings that happen to overlap with ObjectNames that are remapped
+				{
+					FString Actual(TEXT(R"("SrcPackage")"));
+					const FString Expected(TEXT(R"("SrcPackage")"));
+					CHECK(!Patcher.DoPatch(Actual));
+					CHECK_EQUALS(TEXT("Do not remap a string that matches a non-fully-qualified ObjectName"), Actual, Expected);
+				}
 			}
 
 			SECTION("Substring match")
@@ -3157,6 +3181,20 @@ TEST_CASE_NAMED(FAssetHeaderPatcherTests, "AssetHeaderPatcher", "[AssetHeaderPat
 					CHECK(Patcher.DoPatch(Actual));
 					CHECK_EQUALS(TEXT("Patch substring with quoted package, object and sub-object paths"), Actual, Expected);
 				}
+
+				// Do not replace strings that happen to overlap with ObjectNames that are remapped(i.e "SrcPackage=" is not transformed to "DstPackage=") 
+				{
+					FString Actual(
+						TEXT(R"(((SrcPackage="/SrcMount/SomePath/SrcPackage.SrcPackage:RigVMModel.Setup Arm",)")
+						TEXT(R"(((SrcPackage="/SrcMount/SomePath/SrcPackage",)")
+						TEXT(R"(SrcPackage="/SrcMount/SomePath/SrcPackage.SrcPackage_C")))"));
+					FString Expected(
+						TEXT(R"(((SrcPackage="/DstMount/SomePath/DstPackage.DstPackage:RigVMModel.Setup Arm",)")
+						TEXT(R"(((SrcPackage="/DstMount/SomePath/DstPackage",)")
+						TEXT(R"(SrcPackage="/DstMount/SomePath/DstPackage.DstPackage_C")))"));
+					CHECK(Patcher.DoPatch(Actual));
+					CHECK_EQUALS(TEXT("Patch substring with quoted package, object and sub-object paths. No non-fully-qualified ObjectNames are patched."), Actual, Expected);
+				}
 			}
 		}
 
