Allow path length validation to respect project-specific path prefixes
Committing on behalf of Derek.Ehrman

diff --git a/Engine/Source/Developer/AssetTools/Private/AssetViewUtils.cpp b/Engine/Source/Developer/AssetTools/Private/AssetViewUtils.cpp
--- a/Engine/Source/Developer/AssetTools/Private/AssetViewUtils.cpp
+++ b/Engine/Source/Developer/AssetTools/Private/AssetViewUtils.cpp
@@ -1556,7 +1556,7 @@ FString AssetViewUtils::GetPackagePathWithinRoot(const FString& PackageName)
 	return RelativePathToAsset;
 }
 
-int32 AssetViewUtils::GetPackageLengthForCooking(const FString& PackageName, bool bIsInternalBuild)
+int32 AssetViewUtils::GetPackageLengthForCooking(const FString& PackageName, bool bIsInternalBuild, FString CookSubPath)
 {
 	FString RelativePathToAsset;
 	if (!FPackageName::TryConvertLongPackageNameToFilename(PackageName, RelativePathToAsset, FPackageName::GetAssetPackageExtension()))
@@ -1569,9 +1569,17 @@ int32 AssetViewUtils::GetPackageLengthForCooking(const FString& PackageName, boo
 	const TSharedPtr<IPlugin> PluginContainingAsset = GetPluginForFolder(PackageName);
 	const bool bIsEngineAsset = IsEngineFolder(PackageName) || (PluginContainingAsset && PluginContainingAsset->GetLoadedFrom() == EPluginLoadedFrom::Engine);
 	const bool bIsProjectAsset = !bIsEngineAsset;
+	
+	// If the user didn't supply a custom cook sub path, we use "Android_ASTCClient" below as it's probably the longest platform name, so will also prove that any shorter platform names will validate correctly
+	if (CookSubPath.IsEmpty())
+	{
+		CookSubPath = FPaths::Combine(TEXT("Saved"), TEXT("Cooked"), TEXT("Android_ASTCClient"), bIsEngineAsset ? TEXT("Engine") : GameName, TEXT("")); // Trailing empty entry ensures a trailing slash
+	}	
+	else
+	{
+		CookSubPath = FPaths::Combine(CookSubPath, bIsEngineAsset ? TEXT("Engine") : GameName, TEXT("")); // Trailing empty entry ensures a trailing slash
+	}
 
-	// We use "Android_ASTCClient" below as it's probably the longest platform name, so will also prove that any shorter platform names will validate correctly
-	const FString CookSubPath = FPaths::Combine(TEXT("Saved"), TEXT("Cooked"), TEXT("Android_ASTCClient"), bIsEngineAsset ? TEXT("Engine") : GameName, TEXT("")); // Trailing empty entry ensures a trailing slash
 	const FString AbsolutePathToAsset = FPaths::ConvertRelativePathToFull(RelativePathToAsset);
 
 	FString AbsoluteTargetPath = FPaths::ConvertRelativePathToFull(bIsEngineAsset ? FPaths::EngineDir() : FPaths::ProjectDir());
diff --git a/Engine/Source/Developer/AssetTools/Public/AssetViewUtils.h b/Engine/Source/Developer/AssetTools/Public/AssetViewUtils.h
--- a/Engine/Source/Developer/AssetTools/Public/AssetViewUtils.h
+++ b/Engine/Source/Developer/AssetTools/Public/AssetViewUtils.h
@@ -235,7 +235,7 @@ namespace AssetViewUtils
 	ASSETTOOLS_API FString GetPackagePathWithinRoot(const FString& PackageName);
 
 	/** Returns the length of the computed cooked package name and path whether it's run on a build machine or locally */
-	ASSETTOOLS_API int32 GetPackageLengthForCooking(const FString& PackageName, bool bIsInternalBuild);
+	ASSETTOOLS_API int32 GetPackageLengthForCooking(const FString& PackageName, bool bIsInternalBuild, FString CookSubPath={});
 
 	/** Checks to see whether the path is within the size restrictions for cooking */
 	ASSETTOOLS_API bool IsValidPackageForCooking(const FString& PackageName, FText& OutErrorMessage);
