[AssetHeaderPatcher]

- Removes global CoreRedirect use from the patcher and instead each task thread will push patcher redirects for the duration of the task. This is to prevent the patcher from affecting global redirects happening on other threads (adding in irrelevant redirects but also prevents removing important redirects unintentionally)
- Adds copy constructor for RedirectTypeMap to properly copy FRedirectNameMap and FWildcardData values to avoid contaminating existing scopes

#rnx

diff --git a/Engine/Source/Developer/AssetTools/Private/AssetHeaderPatcher.cpp b/Engine/Source/Developer/AssetTools/Private/AssetHeaderPatcher.cpp
--- a/Engine/Source/Developer/AssetTools/Private/AssetHeaderPatcher.cpp
+++ b/Engine/Source/Developer/AssetTools/Private/AssetHeaderPatcher.cpp
@@ -17,6 +17,7 @@
 #include "Misc/RedirectCollector.h"
 #include "Serialization/LargeMemoryReader.h"
 #include "UObject/CoreRedirects.h"
+#include "UObject/CoreRedirects/CoreRedirectsContext.h"
 #include "UObject/Linker.h"
 #include "UObject/NameTypes.h"
 #include "UObject/ObjectResource.h"
@@ -887,6 +888,9 @@ public:
 
 FAssetHeaderPatcher::EResult FAssetHeaderPatcher::DoPatch(const FString& InSrcAsset, const FString& InDstAsset, const FContext& InContext)
 {
+	FScopeCoreRedirectsContext PatchContext(TEXT("Asset Header Patcher Context"), FCoreRedirectsContext::GetThreadContext());
+	FCoreRedirects::AddRedirectList(InContext.Redirects, TEXT("Asset Header Patcher"));
+
 	FAssetHeaderPatcherInner Inner(InSrcAsset, InDstAsset, InContext.StringReplacements, InContext.StringMountReplacements);
 
 	if (!FFileHelper::LoadFileToArray(Inner.SrcBuffer, *Inner.SrcAsset))
@@ -929,9 +933,6 @@ UE::Tasks::FTask FAssetHeaderPatcher::PatchAsync(int32* InOutNumFilesToPatch, in
 		*InOutNumFilesToPatch = PatchedFiles.Num();
 	}
 
-	// Before we start patching we need to apply any patching redirects that exist
-	FCoreRedirects::AddRedirectList(Context.Redirects, TEXT("Asset Header Patcher"));
-
 	// Spawn tasks (Scatter)
 	UE::Tasks::FTask PatchAssetsCleanupTask;
 	TArray<UE::Tasks::FTask> PatchAssetTasks;
@@ -993,8 +994,6 @@ UE::Tasks::FTask FAssetHeaderPatcher::PatchAsync(int32* InOutNumFilesToPatch, in
 	// Once all tasks have completed, remove the redirects before we declare Patching complete
 	UE::Tasks::FTask PatcherCleanupTask = UE::Tasks::Launch(UE_SOURCE_LOCATION, [this]()
 		{
-			FCoreRedirects::RemoveRedirectList(Context.Redirects, TEXT("Asset Header Patcher"));
-
 			if (Status != EResult::Cancelled && ErroredFiles.IsEmpty())
 			{
 				Status = EResult::Success;
diff --git a/Engine/Source/Runtime/CoreUObject/Private/UObject/CoreRedirects.cpp b/Engine/Source/Runtime/CoreUObject/Private/UObject/CoreRedirects.cpp
--- a/Engine/Source/Runtime/CoreUObject/Private/UObject/CoreRedirects.cpp
+++ b/Engine/Source/Runtime/CoreUObject/Private/UObject/CoreRedirects.cpp
@@ -1258,6 +1258,25 @@ int FCoreRedirectObjectName::Compare(const FCoreRedirectObjectName& Other) const
 	return 0;
 }
 
+FCoreRedirects::FRedirectTypeMap::FRedirectTypeMap(const FRedirectTypeMap& Other)
+{
+	FastIterable = Other.FastIterable;
+
+	Map.Reserve(FastIterable.Num());
+	for (TPair<ECoreRedirectFlags, FRedirectNameMap>& Pair : FastIterable)
+	{
+		FRedirectNameMap& RedirectNameMap = Pair.Value;
+		Map.Add(Pair.Key, &RedirectNameMap);
+
+		if (RedirectNameMap.Wildcards)
+		{
+			FWildcardData* WildCardCopy = new FWildcardData();
+			*WildCardCopy = *RedirectNameMap.Wildcards;
+			RedirectNameMap.Wildcards = MakeShareable(WildCardCopy);
+		}
+	}
+}
+
 FCoreRedirects::FRedirectNameMap& FCoreRedirects::FRedirectTypeMap::FindOrAdd(ECoreRedirectFlags Key)
 {
 	FRedirectNameMap*& NameMap = Map.FindOrAdd(Key);
diff --git a/Engine/Source/Runtime/CoreUObject/Public/UObject/CoreRedirects.h b/Engine/Source/Runtime/CoreUObject/Public/UObject/CoreRedirects.h
--- a/Engine/Source/Runtime/CoreUObject/Public/UObject/CoreRedirects.h
+++ b/Engine/Source/Runtime/CoreUObject/Public/UObject/CoreRedirects.h
@@ -474,12 +474,18 @@ private:
 	struct FRedirectTypeMap
 	{
 	public:
+		FRedirectTypeMap() = default;
+		FRedirectTypeMap(const FRedirectTypeMap& Other);
 		FRedirectNameMap& FindOrAdd(ECoreRedirectFlags Key);
 		FRedirectNameMap* Find(ECoreRedirectFlags Key);
 		void Empty();
 
 		TArray<TPair<ECoreRedirectFlags, FRedirectNameMap>>::RangedForIteratorType begin() { return FastIterable.begin(); }
 		TArray<TPair<ECoreRedirectFlags, FRedirectNameMap>>::RangedForIteratorType end() { return FastIterable.end(); }
+
+		FRedirectTypeMap(const FRedirectTypeMap&& Other) = delete;
+		FRedirectTypeMap& operator=(const FRedirectTypeMap& Other) = delete;
+		FRedirectTypeMap& operator=(const FRedirectTypeMap&& Other) = delete;
 	private:
 		TMap<ECoreRedirectFlags, FRedirectNameMap*> Map;
 		TArray<TPair<ECoreRedirectFlags, FRedirectNameMap>> FastIterable;
