Fix for error in logic when patching quoted paths during asset header patching where there are multiple quoted entries in a single string

diff --git a/Engine/Source/Developer/AssetTools/Private/AssetHeaderPatcher.cpp b/Engine/Source/Developer/AssetTools/Private/AssetHeaderPatcher.cpp
--- a/Engine/Source/Developer/AssetTools/Private/AssetHeaderPatcher.cpp
+++ b/Engine/Source/Developer/AssetTools/Private/AssetHeaderPatcher.cpp
@@ -1555,7 +1555,7 @@ bool FAssetHeaderPatcherInner::DoPatch(FString& InOutString)
 
 						// Keep searching until the path is depleted since there might be more than one path to replace
 						// Skip to the end quote since we know we didn't have any matches within the quotes beyond the mount
-						PathView = FStringView(*StringToPatch + SecondQuotePos + 1);
+						PathView = FStringView(*StringToPatch + SecondQuotePos + 1 +(DstMountPoint.Len() - SrcMountPoint.Len())); // Dst - Src to account for new SecondQuotePos after replacement
 					}
 					else
 					{
@@ -2869,8 +2869,8 @@ TEST_CASE_NAMED(FAssetHeaderPatcherTests, "AssetHeaderPatcher", "[AssetHeaderPat
 	const TCHAR*  DstPackagePath			= TEXT("/DstMount/SomePath/DstPackage");
 	const TCHAR*  SrcPackageObjectPath		= TEXT("/SrcMount/SomePath/SrcPackage.SrcPackage");
 	const TCHAR*  DstPackageObjectPath		= TEXT("/DstMount/SomePath/DstPackage.DstPackage");
-	const TCHAR*  SrcMountName				= TEXT("/SrcSpecialMount/");
-	const TCHAR*  DstMountName				= TEXT("/DstSpecialMount/");
+	const TCHAR*  SrcMountName				= TEXT("/SourceSpecialMount/");
+	const TCHAR*  DstMountName				= TEXT("/DestinationSpecialMount/");
 	const FSoftObjectPath SoftObjectPathToRedirect(TEXT("/ToBeRedirectedMount/SomePath/ToBeRedirectedPackage.ToBeRedirectedPackage:Some.ToBeRedirectedPackage.Subobject"));
 	const FSoftObjectPath RedirectedSoftObjectPath(TEXT("/RedirectedMount/SomePath/RedirectedPackage.RedirectedPackage:Some.RedirectedPackage.Subobject"));
 	const FName   SrcPackagePathFName(SrcPackagePath);
@@ -3207,54 +3207,65 @@ TEST_CASE_NAMED(FAssetHeaderPatcherTests, "AssetHeaderPatcher", "[AssetHeaderPat
 			*/
 			/*
 			{
-				FString Actual(TEXT(R"(/SrcSpecialMount/SomePath/SomePackage)"));
-				const FString Expected(TEXT(R"(/DstSpecialMount/SomePath/SomePackage)"));
+				FString Actual(TEXT(R"(/SourceSpecialMount/SomePath/SomePackage)"));
+				const FString Expected(TEXT(R"(/DestinationSpecialMount/SomePath/SomePackage)"));
 				CHECK(Patcher.DoPatch(Actual));
 				CHECK_EQUALS(TEXT("Patch package path replaces only mount"), Actual, Expected);
 			}
 
 			{
-				FString Actual(TEXT(R"(/SrcSpecialMount/SomePath/SomePackage.SomePackage)"));
-				const FString Expected(TEXT(R"(/DstSpecialMount/SomePath/SomePackage.SomePackage)"));
+				FString Actual(TEXT(R"(/SourceSpecialMount/SomePath/SomePackage.SomePackage)"));
+				const FString Expected(TEXT(R"(/DestinationSpecialMount/SomePath/SomePackage.SomePackage)"));
 				CHECK(Patcher.DoPatch(Actual));
 				CHECK_EQUALS(TEXT("Patch package object path replaces only mount"), Actual, Expected);
 			}
 			*/
 
 			{
-				FString Actual(TEXT(R"(/SrcSpecialMount/SomePath/SomePackage.TopLevel:SubObject1.SubObject2)"));
-				const FString Expected(TEXT(R"(/DstSpecialMount/SomePath/SomePackage.TopLevel:SubObject1.SubObject2)"));
+				FString Actual(TEXT(R"(/SourceSpecialMount/SomePath/SomePackage.TopLevel:SubObject1.SubObject2)"));
+				const FString Expected(TEXT(R"(/DestinationSpecialMount/SomePath/SomePackage.TopLevel:SubObject1.SubObject2)"));
 				CHECK(Patcher.DoPatch(Actual));
 				CHECK_EQUALS(TEXT("Patch package sub-object path replaces only mount"), Actual, Expected);
 			}
 
 			{
-				FString Actual(TEXT(R"("/SrcSpecialMount/SomePath/SomePackage")"));
-				const FString Expected(TEXT(R"("/DstSpecialMount/SomePath/SomePackage")"));
+				FString Actual(TEXT(R"("/SourceSpecialMount/SomePath/SomePackage")"));
+				const FString Expected(TEXT(R"("/DestinationSpecialMount/SomePath/SomePackage")"));
 				CHECK(Patcher.DoPatch(Actual));
 				CHECK_EQUALS(TEXT("Patch double quoted path replaces only mount"), Actual, Expected);
 			}
 
 			{
-				FString Actual(TEXT(R"('/SrcSpecialMount/SomePath/SomePackage')"));
-				const FString Expected(TEXT(R"('/DstSpecialMount/SomePath/SomePackage')"));
+				FString Actual(TEXT(R"('/SourceSpecialMount/SomePath/SomePackage')"));
+				const FString Expected(TEXT(R"('/DestinationSpecialMount/SomePath/SomePackage')"));
 				CHECK(Patcher.DoPatch(Actual));
 				CHECK_EQUALS(TEXT("Patch single quoted path replaces only mount"), Actual, Expected);
 			}
 
 			{
-				FString Actual(TEXT(R"(SomePrefix="/SrcSpecialMount/SomePath/SomePackage")"));
-				const FString Expected(TEXT(R"(SomePrefix="/DstSpecialMount/SomePath/SomePackage")"));
+				FString Actual(TEXT(R"(SomePrefix="/SourceSpecialMount/SomePath/SomePackage")"));
+				const FString Expected(TEXT(R"(SomePrefix="/DestinationSpecialMount/SomePath/SomePackage")"));
 				CHECK(Patcher.DoPatch(Actual));
 				CHECK_EQUALS(TEXT("Substring patch replaces only mount when double quoted"), Actual, Expected);
 			}
 
 			{
-				FString Actual(TEXT(R"(SomePrefix='/SrcSpecialMount/SomePath/SomePackage')"));
-				const FString Expected(TEXT(R"(SomePrefix='/DstSpecialMount/SomePath/SomePackage')"));
+				FString Actual(TEXT(R"(SomePrefix='/SourceSpecialMount/SomePath/SomePackage')"));
+				const FString Expected(TEXT(R"(SomePrefix='/DestinationSpecialMount/SomePath/SomePackage')"));
 				CHECK(Patcher.DoPatch(Actual));
 				CHECK_EQUALS(TEXT("Substring patch replaces only mount when single quoted"), Actual, Expected);
 			}
+
+			{
+				FString Actual(
+					TEXT(R"("/SourceSpecialMount/SomePath/SomePackage1",)")
+					TEXT(R"("/SourceSpecialMount/SomePath/SomePackage2")"));
+				const FString Expected(
+					TEXT(R"("/DestinationSpecialMount/SomePath/SomePackage1",)")
+					TEXT(R"("/DestinationSpecialMount/SomePath/SomePackage2")"));
+				CHECK(Patcher.DoPatch(Actual));
+				CHECK_EQUALS(TEXT("Substring patch replaces mount in multiple double quoted paths"), Actual, Expected);
+			}
 		}
 	}
 
