[AssetHeaderPatcher]

When initializing the patcher context rfom package paths, deduce the mountpoint replacement mapping to be used as prefix redirects

#rnx

diff --git a/Engine/Source/Developer/AssetTools/Private/AssetHeaderPatcher.cpp b/Engine/Source/Developer/AssetTools/Private/AssetHeaderPatcher.cpp
--- a/Engine/Source/Developer/AssetTools/Private/AssetHeaderPatcher.cpp
+++ b/Engine/Source/Developer/AssetTools/Private/AssetHeaderPatcher.cpp
@@ -225,6 +225,21 @@ void FAssetHeaderPatcher::FContext::GatherDependentPackages()
 			// add the local path/asset for the case of maps (which we cannot tell at this point)
 			ExternalMappings.Add(FString(SrcPath.GetData()), FString(DstPath.GetData()));
 		}
+
+		// While iterating mappings, add any mountpoint changes
+		if (SrcPackageRoot != DstPackageRoot)
+		{
+			const uint32 SrcPackageRootHash = GetTypeHash(SrcPackageRoot);
+			FString* RemappedRoot = StringMountReplacements.FindByHash(SrcPackageRootHash, SrcPackageRoot);
+			if (RemappedRoot)
+			{
+				UE_CLOG(DstPackageRoot != *RemappedRoot, LogAssetHeaderPatcher, Warning,
+					TEXT("Found conflicting mountpoint remapping: /%s/ -> /%s/ and /%s/ -> /%s/. The second mapping will be used to overwrite the first."),
+					*FString(SrcPackageRoot), **RemappedRoot,
+					*FString(SrcPackageRoot), *FString(DstPackageRoot));
+			}
+			StringMountReplacements.AddByHash(SrcPackageRootHash, FString(SrcPackageRoot), FString(DstPackageRoot));
+		}
 	}
 
 	TMap<FString, FString> Result;
