AssetViewUtils::LoadAssetsIfNeeded will now always to try to load the optional outer asset first, and also load the required objects by their package name instead of their long name.
This is required when fixing redirectors to public objects inside assets, which wasn't possible before having public entities.

diff --git a/Engine/Source/Developer/AssetTools/Private/AssetViewUtils.cpp b/Engine/Source/Developer/AssetTools/Private/AssetViewUtils.cpp
--- a/Engine/Source/Developer/AssetTools/Private/AssetViewUtils.cpp
+++ b/Engine/Source/Developer/AssetTools/Private/AssetViewUtils.cpp
@@ -192,28 +192,25 @@ AssetViewUtils::ELoadAssetsResult AssetViewUtils::LoadAssetsIfNeeded(TConstArray
 	bool bAnySucceeded = false;
 
 	// Build a list of unloaded assets
-	TArray<FSoftObjectPath> UnloadedObjectPaths;
-	bool bAtLeastOneUnloadedMap = false;
+	TSet<FName> UnloadedObjectPaths;
 	for (const FAssetData& Asset : Assets)
 	{
-		UObject* FoundObject = Asset.FastGetAsset(false);
-		if (FoundObject)
+		if (UObject* FoundObject = Asset.FastGetAsset(false))
 		{
 			OutLoadedObjects.Add(FoundObject);
 			bAnySucceeded = true;
 			continue;
 		}
-		else if (!Settings.bLoadWorldPartitionMaps && ULevel::GetIsLevelPartitionedFromAsset(Asset))
-		{
-			// Skip
-			bAtLeastOneUnloadedMap = true;
-		}
-		else
+
+		if (Settings.bLoadWorldPartitionMaps || !ULevel::GetIsLevelPartitionedFromAsset(Asset))
 		{
-			UnloadedObjectPaths.Add(Asset.GetSoftObjectPath());
+			UnloadedObjectPaths.Add(Asset.GetOptionalOuterPathName());
+			UnloadedObjectPaths.Add(Asset.PackageName);
 		}
 	}
 
+	UnloadedObjectPaths.Remove(NAME_None);
+
 	if (UnloadedObjectPaths.Num() == 0)
 	{
 		return ELoadAssetsResult::Success;
@@ -250,9 +247,9 @@ AssetViewUtils::ELoadAssetsResult AssetViewUtils::LoadAssetsIfNeeded(TConstArray
 		TGuardValueAccessors<bool> IsEditorLoadingPackageGuard(UE::GetIsEditorLoadingPackage, UE::SetIsEditorLoadingPackage, true);
 		const ELoadFlags LoadFlags = Settings.bFollowRedirectors ? LOAD_None : LOAD_NoRedirects;
 
-		for (const FSoftObjectPath& ObjectPath : UnloadedObjectPaths)
+		for (const FName& ObjectPath : UnloadedObjectPaths)
 		{
-			SlowTask.EnterProgressFrame(1, FText::Format(LOCTEXT("LoadingAsset", "Loading {0}..."), FText::FromName(ObjectPath.GetLongPackageFName())));
+			SlowTask.EnterProgressFrame(1, FText::Format(LOCTEXT("LoadingAsset", "Loading {0}..."), FText::FromName(ObjectPath)));
 
 			// Load up the object
 			FLinkerInstancingContext InstancingContext(Settings.bLoadAllExternalObjects ? TSet<FName>{ ULevel::LoadAllExternalObjectsTag } : TSet<FName>());
