[AssetHeaderPatching]

- Adds support to allow redirects to add to the ImportTable.
- Extends header serializer to support growing the ImportTable when patching rather than only inline patches.
- Prevent stomping over FNames in the NameTable when object paths made from multiple FNames share names with other objects that do not need any patching. We now track which names are used so we can ensure remapped names either replace existing names or are treated as an 'added' name.
- Adds patching remapping for the fullpath to subobjects (using a wildcard redirect) so that we can created the patched form for subobject fullpaths. Sub-object paths are only important because imports can chain from CDO subobjects as well as externalactor/ExternalObject types and we need to be able to form their fullpath to establish their patching context as well as know which parts should be patched.
- Added GetTypeHash() for FCoreRedirectObjectName so the type can be used in Maps
- Adds Tests for ImportTable patching

#rnx

diff --git a/Engine/Source/Developer/AssetTools/Private/AssetHeaderPatcher.cpp b/Engine/Source/Developer/AssetTools/Private/AssetHeaderPatcher.cpp
--- a/Engine/Source/Developer/AssetTools/Private/AssetHeaderPatcher.cpp
+++ b/Engine/Source/Developer/AssetTools/Private/AssetHeaderPatcher.cpp
@@ -444,6 +444,26 @@ void FAssetHeaderPatcher::FContext::GenerateAdditionalRemappings()
 			Redirects.Emplace(MoveTemp(PackageObjectRedirect));
 		}
 
+		// Path.ObjectName.* mapping
+		{
+			SrcNameBuilder.Reset();
+			SrcNameBuilder.Append(SrcNameString);
+			SrcNameBuilder.AppendChar(TEXT('.'));
+			SrcNameBuilder.Append(SrcPackageName);
+			SrcNameBuilder.AppendChar(TEXT('.'));
+
+			DstNameBuilder.Reset();
+			DstNameBuilder.Append(DstNameString);
+			DstNameBuilder.AppendChar(TEXT('.'));
+			DstNameBuilder.Append(DstPackageName);
+			DstNameBuilder.AppendChar(TEXT('.'));
+
+			FCoreRedirect PackageObjectRedirect(ECoreRedirectFlags::Option_MatchPrefix | ECoreRedirectFlags::Type_Object,
+				FCoreRedirectObjectName(SrcNameBuilder.ToString()),
+				FCoreRedirectObjectName(DstNameBuilder.ToString()));
+			Redirects.Emplace(MoveTemp(PackageObjectRedirect));
+		}
+
 		// MaterialFunctionInterface "EditorOnlyData"
 		{
 			SrcNameBuilder.Reset();
@@ -819,15 +839,59 @@ public:
 		}
 	}
 
+	// Reset anything not set via construction. Used for testing
+	void ResetInternalState()
+	{
+		AssetRegistryData = FAssetRegistryData();
+		HeaderInformation = FHeaderInformation();
+		AddedNames.Empty();
+		ExportTable.Empty();
+		ImportTable.Empty();
+		ImportTablePatchedNames.Empty();
+		ImportNameToImportTableIndexLookup.Empty();
+		GatherableTextDataTable.Empty();
+		NameTable.Empty();
+		NameToIndexMap.Empty();
+		RenameMap.Empty();
+		SearchableNamesMap.Empty();
+		SoftObjectPathTable.Empty();
+		SoftPackageReferencesTable.Empty();
+		Summary = FPackageFileSummary();
+		ThumbnailTable.Empty();
+		UnchangedNames.Empty();
+	}
+
 	bool DoPatch(FString& InOutString);
 	bool DoPatch(FName& InOutName);
-	bool DoPatch(FObjectResource& InOutResource, bool bIsExport, FName& OutPatchedObjectName);
 	bool DoPatch(FSoftObjectPath& InOutSoft);
 	bool DoPatch(FTopLevelAssetPath& InOutPath);
 	bool DoPatch(FGatherableTextData& InOutGatherablerTextData);
 	bool DoPatch(FThumbnailEntry& InOutThumbnailEntry);
 	bool RemapFName(FName SrcName, FName DstName);
 	bool ShouldReplaceMountPoint(const FStringView InPath, FStringView& OutSrcMountPoint, FStringView& OutDstMountPoint);
+	FCoreRedirectObjectName GetFullObjectNameFromObjectResource(const FObjectResource& InResource, bool bIsExport, bool bWalkImportsOnly = false);
+
+	struct FExportPatch
+	{
+		int32 TableIndex;
+		FName ObjectName;
+	};
+
+	struct FImportPatch
+	{
+		int32 TableIndex;
+		FName ObjectName;
+		FPackageIndex OuterIndex;
+		FName ClassName;
+		FName ClassPackage;
+#if WITH_EDITORONLY_DATA
+		FName PackageName;
+#endif
+	};
+
+	void GetExportTablePatches(TArray<FExportPatch>& OutExportPatches);
+	void GetImportTablePatches(TArray<FImportPatch>& OutImportPatches, int32& OutNewImportCount);
+	void PatchExportAndImportTables(const TArray<FExportPatch>& InExportPatches, const TArray<FImportPatch>& InImportPatches, const int32 InNewImportCount);
 	void PatchNameTable();
 
 	FAssetHeaderPatcher::EResult PatchHeader();
@@ -874,8 +938,12 @@ public:
 	// NameTable Members
 	TArray<FName> NameTable;
 	TMap<FNameEntryId, int32> NameToIndexMap;
+	TSet<FNameEntryId> UnchangedNames;
 	TMap<FNameEntryId, FNameEntryId> RenameMap;
 	TSet<FNameEntryId> AddedNames;
+	// Export/Import table
+	TArray<TPair<FCoreRedirectObjectName, FCoreRedirectObjectName>> ImportTablePatchedNames;
+	TMap<FCoreRedirectObjectName, int32> ImportNameToImportTableIndexLookup;
 
 	TArray<FSoftObjectPath> SoftObjectPathTable;
 	TArray<FGatherableTextData> GatherableTextDataTable;
@@ -1219,6 +1287,7 @@ FAssetHeaderPatcher::EResult FAssetHeaderPatcherInner::PatchHeader_Deserialize()
 		// value but might not after patching (i.e. their use of the name differs based on context, and
 		// post-patching the FNames in those contexts no longer match.
 		NameToIndexMap.Empty(NameTable.Num());
+		UnchangedNames.Reserve(NameTable.Num());
 		RenameMap.Reserve(NameTable.Num());
 		AddedNames.Empty();
 		for (int32 i = 0; i < NameTable.Num(); ++i)
@@ -1287,6 +1356,8 @@ FAssetHeaderPatcher::EResult FAssetHeaderPatcherInner::PatchHeader_Deserialize()
 
 		MemAr.Seek(Summary.ImportOffset);
 		ImportTable.Reserve(Summary.ImportCount);
+		ImportTablePatchedNames.Reserve(ImportTable.Num());
+		ImportNameToImportTableIndexLookup.Reserve(ImportTable.Num());
 		for (int32 ImportIndex = 0; ImportIndex < Summary.ImportCount; ++ImportIndex)
 		{
 			FObjectImport& Import = ImportTable.Emplace_GetRef();
@@ -1613,36 +1684,60 @@ bool FAssetHeaderPatcherInner::DoPatch(FString& InOutString)
 
 bool FAssetHeaderPatcherInner::RemapFName(FName SrcName, FName DstName)
 {
+	// None won't appear in the NameTable so skip any None passed in here
+	if(SrcName == NAME_None)
+	{
+		return false;
+	}
+	checkf(DstName != NAME_None, TEXT("There should never be a None FName in the NameTable"));
+
 	// NameTable entries only care about the comparison form (no number) so 
 	// only consider that for remapping purposes
 	FNameEntryId SrcComparisonId = SrcName.GetDisplayIndex();
 	FNameEntryId DstComparisonId = DstName.GetDisplayIndex();
-	if (SrcComparisonId == DstComparisonId)
+
+	// Since we do fuzzy matching against AssetRegistryTag data, we might get SrcNames not in the FName table as input; ignore these.
+	if(!NameToIndexMap.Contains(SrcComparisonId))
 	{
 		return false;
 	}
-	checkf(DstName != NAME_None, TEXT("There should never be a None FName in the NameTable"));
 
+	bool bIsExportDataFName = NameToIndexMap[SrcComparisonId] < Summary.NamesReferencedFromExportDataCount;
+	bool bNeedRemap = SrcComparisonId != DstComparisonId;
 	FNameEntryId* RemappedFName = RenameMap.Find(SrcComparisonId);
-	if (RemappedFName)
+
+	if (!bNeedRemap)
 	{
-		// We already have a mapping. That is fine; we might have used the same FName in more than one place.
-		// However, we need to be certain we are renaming the name to the same new name. If not, this means
-		// the originals names overlapped but in the patched case they don't (e.g. A class FName may have matched a Package
-		// name, but after patching it's possible _only_ the Package name has changed. In such a case we don't want to rename the 
-		// class name inadvertently by patching the shared NameTable entry. If we have a mismatch with the new patched 
-		// name, record the new name and we will append it to the NameTable later.
-		if (*RemappedFName != DstComparisonId)
+		UnchangedNames.Add(SrcComparisonId);
+
+		if (RemappedFName)
 		{
-			AddedNames.Add(DstComparisonId);
+			// We hit a case where we thought we needed to change the name in the NameTable
+			// but have now discovered some part of the header needs to use the old name. In such a case
+			// remove the remap and make it an add instead.
+			AddedNames.Add(*RemappedFName);
+			RenameMap.Remove(SrcComparisonId);
 		}
 	}
 	else
 	{
-		RenameMap.Add(SrcComparisonId, DstComparisonId);
+		// We already have a _different_ remapping, a header entry is still relying on the original name or this is a name used by the Export payload in which case we can't safely modify it. That is fine; we might have used the same FName in more than one place.
+		// We need to be certain we are renaming only the names we care about in context. If not, shared names in the NameTable might change
+		// incorrectly (e.g. A class FName may have matched a Package FName, thus sharing a NameTable entry, but after patching it's possible _only_ 
+		// the Package name has changed. In such a case we don't want to rename the class name inadvertently by patching the shared NameTable entry. 
+		// If we have a mismatch with the new patched name, record the new name and we will append it to the NameTable later.
+		bool bForceAddName = bIsExportDataFName || (RemappedFName && *RemappedFName != DstComparisonId) || UnchangedNames.Contains(SrcComparisonId);
+		if (bForceAddName)
+		{
+			AddedNames.Add(DstComparisonId);
+		}
+		else
+		{
+			RenameMap.Add(SrcComparisonId, DstComparisonId);
+		}
 	}
 
-	return true;
+	return bNeedRemap;
 }
 
 bool FAssetHeaderPatcherInner::DoPatch(FName& InOutName)
@@ -1659,6 +1754,243 @@ bool FAssetHeaderPatcherInner::DoPatch(FName& InOutName)
 	return false;
 }
 
+void FAssetHeaderPatcherInner::GetExportTablePatches(TArray<FExportPatch>& OutExportPatches)
+{
+	OutExportPatches.Reserve(ExportTable.Num());
+	for (int32 i = 0; i < ExportTable.Num(); ++i)
+	{
+		FObjectExport& Export = ExportTable[i];
+		FCoreRedirectObjectName SrcResourceName = GetFullObjectNameFromObjectResource(Export, true /*bIsExport*/);
+		FCoreRedirectObjectName DstResourceName = FCoreRedirects::GetRedirectedName(ECoreRedirectFlags::Type_AllMask, SrcResourceName);
+		if (RemapFName(SrcResourceName.ObjectName, DstResourceName.ObjectName))
+		{
+			FExportPatch ExportPatch;
+			ExportPatch.TableIndex = i;
+			ExportPatch.ObjectName = DstResourceName.ObjectName;
+			OutExportPatches.Add(MoveTemp(ExportPatch));
+		}
+	}
+}
+
+void FAssetHeaderPatcherInner::GetImportTablePatches(TArray<FImportPatch>& OutImportPatches, int32& OutNewImportCount)
+{
+	OutNewImportCount = 0;
+	OutImportPatches.Reserve(ImportTable.Num());
+	
+	TArray<int32> PatchedImports;
+	PatchedImports.Reserve(ImportTable.Num());
+
+	for (int32 i = 0; i < ImportTable.Num(); ++i)
+	{
+		FObjectImport& Import = ImportTable[i];
+
+		FCoreRedirectObjectName SrcResourceName = GetFullObjectNameFromObjectResource(Import, false /*bIsExport*/, true /*bWalkImportsOnly*/);
+		FCoreRedirectObjectName DstResourceName = FCoreRedirects::GetRedirectedName(ECoreRedirectFlags::Type_AllMask, SrcResourceName);
+		ImportNameToImportTableIndexLookup.Add(SrcResourceName, i);
+		ImportTablePatchedNames.Add({ SrcResourceName, DstResourceName });
+
+		if (SrcResourceName != DstResourceName)
+		{
+			PatchedImports.Add(i);
+		}
+
+		// While we are here we patch FNames that have no impact on the FObjectResource walking we will do below.
+		// This saves us from having to output FImportPatch entries for entries that don't have their ObjectPath patched
+		// but may have changed the package name or class path.
+		const FCoreRedirectObjectName SrcImportClass(Import.ClassName, NAME_None, Import.ClassPackage);
+		const FCoreRedirectObjectName DstImportClass = FCoreRedirects::GetRedirectedName(ECoreRedirectFlags::Type_Class | ECoreRedirectFlags::Type_Package, SrcImportClass);
+
+		Import.ClassName = RemapFName(SrcImportClass.ObjectName, DstImportClass.ObjectName)
+			? DstImportClass.ObjectName : SrcImportClass.ObjectName;
+
+		Import.ClassPackage = RemapFName(SrcImportClass.PackageName, DstImportClass.PackageName)
+			? DstImportClass.PackageName : SrcImportClass.PackageName;
+
+#if WITH_EDITORONLY_DATA
+		const FCoreRedirectObjectName SrcImportPackageName(NAME_None, NAME_None, Import.PackageName);
+		const FCoreRedirectObjectName DstImportPackageName = FCoreRedirects::GetRedirectedName(ECoreRedirectFlags::Type_Package, SrcImportPackageName);
+		Import.PackageName = RemapFName(SrcImportPackageName.PackageName, DstImportPackageName.PackageName)
+			? DstImportPackageName.PackageName : SrcImportPackageName.PackageName;
+#endif
+	}
+
+	// Now that we know which import entries won't change we can review our patched Imports and see how they need to be patched
+	// (e.g. did the ObjectName change and/or did the Outer change) and determine if we need new even more new imports (in which case we add them)
+	for (auto It = PatchedImports.CreateIterator(); It; It++)
+	{
+		const int32 ImportIndexToPatch = *It;
+		It.RemoveCurrent();
+
+		// Go outer-most to inner-most from the patched name and see if we have all necessary imports patches.
+		// If we already have the patched Import's name in our lookup table, no patching needed; record the OuterIndex 
+		// for the next inner import outered to this one to use. If not, create a new import patch.
+		FPackageIndex OuterIndex = FPackageIndex();
+		const FCoreRedirectObjectName& SrcResourceName = ImportTablePatchedNames[ImportIndexToPatch].Key;
+		const FCoreRedirectObjectName& DstResourceName = ImportTablePatchedNames[ImportIndexToPatch].Value;
+
+		auto CreatePatch = [this, ImportIndexToPatch, &OuterIndex, &SrcResourceName, 
+			&PatchedImports, &OutImportPatches, &OutNewImportCount] (const FCoreRedirectObjectName SrcImportName, const FCoreRedirectObjectName DstImportName, const FName DstImportObjectName)
+			{
+				check(!DstImportObjectName.IsNone());
+
+				// If we are going to patch this SrcImportName later, skip so we can ensure we don't inadvertently add unnecessary entries
+				int32* SrcIndex = ImportNameToImportTableIndexLookup.Find(SrcImportName);
+				if (SrcIndex && PatchedImports.Contains(*SrcIndex))
+				{
+					// Update our outer so any patched imports for objects outered to this one are linked correctly
+					OuterIndex = FPackageIndex::FromImport(*SrcIndex);
+					return;
+				}
+
+				FImportPatch ImportPatch;
+
+				// Prepopulate the patch with the existing import data. We may have patched fields like
+				// ClassName already and want those to be copied over when patching, or used as the 
+				// new value when adding new entries
+				const FObjectImport& SrcImport = ImportTable[ImportIndexToPatch];
+				ImportPatch.ClassName = SrcImport.ClassName;
+				ImportPatch.ClassPackage = SrcImport.ClassPackage;
+#if WITH_EDITORONLY_DATA
+				ImportPatch.PackageName = SrcImport.PackageName;
+#endif
+
+				// If import entry matches exactly the one we are patching (not just a portion of
+				// it's path) then we can overwrite the existing table entry instead of adding a new one
+				int32 PatchImportTableIndex;
+				if (SrcImportName == SrcResourceName)
+				{
+					PatchImportTableIndex = ImportIndexToPatch;
+				}
+				else
+				{
+					PatchImportTableIndex = ImportTable.Num() + OutNewImportCount++;
+				}
+				
+				ImportPatch.TableIndex = PatchImportTableIndex;
+				ImportPatch.ObjectName = DstImportObjectName;
+				RemapFName(SrcImport.ObjectName, DstImportObjectName);
+				ImportPatch.OuterIndex = OuterIndex;
+
+				// Update our outer so any patched imports for objects outered to this one are linked correctly
+				OuterIndex = FPackageIndex::FromImport(PatchImportTableIndex);
+
+				check(!ImportNameToImportTableIndexLookup.Contains(DstImportName));
+				ImportNameToImportTableIndexLookup.Add(DstImportName, PatchImportTableIndex);
+				OutImportPatches.Add(MoveTemp(ImportPatch));
+			};
+
+		const FCoreRedirectObjectName DstPackageImportName(NAME_None, NAME_None, DstResourceName.PackageName);
+		if (int32* ImportIndex = ImportNameToImportTableIndexLookup.Find(DstPackageImportName))
+		{
+			OuterIndex = FPackageIndex::FromImport(*ImportIndex);
+		}
+		else
+		{
+			const FCoreRedirectObjectName SrcPackageImportName(NAME_None, NAME_None, SrcResourceName.PackageName);
+			CreatePatch(SrcPackageImportName, DstPackageImportName, DstPackageImportName.PackageName);
+		}
+		
+		if (!DstResourceName.OuterName.IsNone())
+		{
+			TStringBuilder<NAME_SIZE> ObjectString;
+			DstResourceName.PackageName.ToString(ObjectString);
+			ObjectString.AppendChar(TEXT('.'));
+			const int32 PackagePrefixLen = ObjectString.Len();
+
+			TStringBuilder<NAME_SIZE> OuterString;
+			DstResourceName.OuterName.ToString(OuterString);
+			FStringView OuterStringView = OuterString.ToView();
+			int32 DelimiterPos = OuterStringView.Len();
+
+			do
+			{
+				ObjectString.RemoveSuffix(ObjectString.Len() - PackagePrefixLen);
+				OuterStringView.LeftInline(DelimiterPos);
+				ObjectString.Append(OuterStringView);
+				FCoreRedirectObjectName DstOuterImportName(ObjectString.ToString());
+				if (int32* ImportIndex = ImportNameToImportTableIndexLookup.Find(DstOuterImportName))
+				{
+					OuterIndex = FPackageIndex::FromImport(*ImportIndex);
+				}
+				else
+				{
+					const FCoreRedirectObjectName SrcOuterImportName(SrcResourceName.OuterName, NAME_None, SrcResourceName.PackageName);
+					CreatePatch(SrcOuterImportName, DstOuterImportName, DstOuterImportName.ObjectName);
+				}
+			} while (OuterStringView.FindLastChar(TEXT('.'), DelimiterPos));
+		}
+
+		FCoreRedirectObjectName DstObjectImportName(DstResourceName.ObjectName, DstResourceName.OuterName, DstResourceName.PackageName);
+		if (ImportNameToImportTableIndexLookup.Find(DstObjectImportName) == nullptr)
+		{
+			const FCoreRedirectObjectName SrcObjectImportName(SrcResourceName.ObjectName, SrcResourceName.OuterName, SrcResourceName.PackageName);
+			CreatePatch(SrcObjectImportName, DstObjectImportName, DstObjectImportName.ObjectName);
+		}
+	}
+}
+
+void FAssetHeaderPatcherInner::PatchExportAndImportTables(const TArray<FExportPatch>& InExportPatches, const TArray<FImportPatch>& InImportPatches, const int32 InNewImportCount)
+{
+	// Any FPackageIndex that IsExport() will not change since we do not add or remove entries from the ExportTable.
+	// For Imports we may have changed import names such that we resulted in new imports and now imports that are no longer
+	// in use and should be removed (to avoid unnecessary overhead when loading). However, as far as I can tell we can't 
+	// determine what to remove based on what is used since we append extra imports during serialization that don't have any
+	// Exports or Imports referring to them (e.g. CDO subobjects are like this). As such we append our import patches as new imports
+	// if we can't stomp over existing entries in the table. We then fixup any FPackageIndex members that might be using old indices that we
+	// know we have explicitly changed due to additions.
+
+	ImportTable.SetNum(ImportTable.Num() + InNewImportCount);
+	for (const FImportPatch& ImportPatch : InImportPatches)
+	{
+		const int32 Index = ImportPatch.TableIndex;
+		check(Index < ImportTable.Num());
+
+		FObjectImport& Import = ImportTable[Index];
+		Import.ObjectName = ImportPatch.ObjectName;
+		Import.OuterIndex = ImportPatch.OuterIndex;
+		Import.ClassName = ImportPatch.ClassName;
+		Import.ClassPackage = ImportPatch.ClassPackage;
+#if WITH_EDITORONLY_DATA
+		Import.OldClassName = NAME_None;
+		Import.PackageName = ImportPatch.PackageName;
+#endif
+	}
+
+	for (const FExportPatch& ExportPatch : InExportPatches)
+	{
+		// Patching cannot result in new Exports
+		const int32 Index = ExportPatch.TableIndex;
+		check(Index < ExportTable.Num());
+
+		FObjectExport& Export = ExportTable[Index];
+		Export.ObjectName = ExportPatch.ObjectName;
+#if WITH_EDITORONLY_DATA
+		Export.OldClassName = NAME_None;
+#endif
+	}
+
+	// Walk through our Exports and ensure any FPackageIndex they have is pointing to the correct location
+	auto RemapIndex = [this](FPackageIndex& Index)
+		{
+			if (Index.IsImport())
+			{
+				FCoreRedirectObjectName PatchedName = ImportTablePatchedNames[Index.ToImport()].Value;
+				Index = FPackageIndex::FromImport(ImportNameToImportTableIndexLookup[PatchedName]);
+			}
+		};
+
+	for (FObjectExport& Export : ExportTable)
+	{
+		RemapIndex(Export.ClassIndex);
+		RemapIndex(Export.SuperIndex);
+		RemapIndex(Export.TemplateIndex);
+		RemapIndex(Export.OuterIndex);
+	}
+
+	// We may have added new Imports so ensure the Summary is accurate
+	Summary.ImportCount = ImportTable.Num();
+}
+
 void FAssetHeaderPatcherInner::PatchNameTable()
 {
 	// Note, no number is assigned when replacing FNames as the NameTable only tracks unnumbered names
@@ -1699,33 +2031,41 @@ bool FAssetHeaderPatcherInner::DoPatch(FSoftObjectPath& InOutSoft)
 
 	// Honour explictly remapped paths first. Even if this succeeds we need to handle redirectors
 	// because FSoftObjectPath serialization will always apply them.
+
+	// Special care is needed to check to see if a patch is done via remapping or via redirects before we assume the 
+	// name should be marked as unchanging (causing any future patching to the name to be an add in the NameTable). If we have a redirect 
+	// but not remapping, the failed remapping could marked the srcname as unchanging erroneously.
 	FTopLevelAssetPath SrcTopLevelAssetPath = InOutSoft.GetAssetPath();
-	FTopLevelAssetPath InOutTopLevelAssetPath = SrcTopLevelAssetPath;
-	bool bPatched = DoPatch(InOutTopLevelAssetPath);
+	const FCoreRedirectObjectName DstTopLevelAssetPathObjectName = FCoreRedirects::GetRedirectedName(ECoreRedirectFlags::Type_AllMask, SrcTopLevelAssetPath);
+	FTopLevelAssetPath DstTopLevelAssetPath(DstTopLevelAssetPathObjectName.ToString());
+	bool bPatched = SrcTopLevelAssetPath != DstTopLevelAssetPath;
 	if (bPatched)
 	{
-		InOutSoft.SetPath(InOutTopLevelAssetPath, InOutSoft.GetSubPathString());
+		InOutSoft.SetPath(DstTopLevelAssetPath, InOutSoft.GetSubPathString());
 	}
 
 // Only run in Editor builds as FSoftObjectPath only runs PreSavePath during SerializePath in editor builds
 #if WITH_EDITOR
 	if (InOutSoft.PreSavePath(nullptr))
 	{
-		FTopLevelAssetPath DstTopLevelAssetPath = InOutSoft.GetAssetPath();
-
-		bPatched |= RemapFName(InOutTopLevelAssetPath.GetAssetName(), DstTopLevelAssetPath.GetAssetName());
-		bPatched |= RemapFName(InOutTopLevelAssetPath.GetPackageName(), DstTopLevelAssetPath.GetPackageName());
-		checkf(bPatched, TEXT("Applied a AssetPathRedirection while patching '%s' -> '%s' but failed to update the nametable."), *SrcTopLevelAssetPath.ToString(), *DstTopLevelAssetPath.ToString());
+		bPatched = true;
+		DstTopLevelAssetPath = InOutSoft.GetAssetPath();
 	}
 #endif
 
+	if (bPatched)
+	{
+		RemapFName(SrcTopLevelAssetPath.GetAssetName(), DstTopLevelAssetPath.GetAssetName());
+		RemapFName(SrcTopLevelAssetPath.GetPackageName(), DstTopLevelAssetPath.GetPackageName());
+	}
+
 	return bPatched;
 }
 
-bool FAssetHeaderPatcherInner::DoPatch(FObjectResource& InOutResource, bool bIsExport, FName& OutPatchedObjectName)
+FCoreRedirectObjectName FAssetHeaderPatcherInner::GetFullObjectNameFromObjectResource(const FObjectResource& InResource, bool bIsExport, bool bWalkImportsOnly)
 {
 	bool bOutermostIsExport = bIsExport;
-	FPackageIndex OuterIndex = InOutResource.OuterIndex;
+	FPackageIndex OuterIndex = InResource.OuterIndex;
 	TArray<FName, TInlineAllocator<8>> OuterStack;
 	while (!OuterIndex.IsNull())
 	{
@@ -1735,6 +2075,10 @@ bool FAssetHeaderPatcherInner::DoPatch(FObjectResource& InOutResource, bool bIsE
 			bOutermostIsExport = false;
 			OuterResource = &ImportTable[OuterIndex.ToImport()];
 		}
+		else if (bWalkImportsOnly)
+		{
+			break;
+		}
 		else
 		{
 			bOutermostIsExport = true;
@@ -1748,22 +2092,22 @@ bool FAssetHeaderPatcherInner::DoPatch(FObjectResource& InOutResource, bool bIsE
 	FName SrcObjectName;
 	FName SrcOuterName;
 	FName SrcPackageName;
-	bool bRemapByPackageName = false;
+	bool bRemapByPackage = false;
 	if (OuterStack.Num() == 0)
 	{
 		if (bOutermostIsExport)
 		{
-			SrcPackageName = OriginalPackagePath;		// /Package/Package
+			SrcPackageName = OriginalPackagePath;	// /Package/Package
 			SrcOuterName = NAME_None;
-			SrcObjectName = InOutResource.ObjectName;	// MyObject
+			SrcObjectName = InResource.ObjectName;	// MyObject
 		}
 		else
 		{
 			// The ObjectName is a package
-			SrcPackageName = InOutResource.ObjectName;	// /Package/Package
+			SrcPackageName = InResource.ObjectName;	// /Package/Package
 			SrcOuterName = NAME_None;
 			SrcObjectName = NAME_None;
-			bRemapByPackageName = true;
+			bRemapByPackage = true;
 		}
 	}
 	else
@@ -1774,7 +2118,7 @@ bool FAssetHeaderPatcherInner::DoPatch(FObjectResource& InOutResource, bool bIsE
 		while (!OuterStack.IsEmpty())
 		{
 			FName Outer = OuterStack.Pop();
-			Outer.ToString(OuterString);
+			Outer.AppendString(OuterString);
 			OuterString.AppendChar(TEXT('.'));
 		}
 		if (OuterString.Len())
@@ -1782,29 +2126,10 @@ bool FAssetHeaderPatcherInner::DoPatch(FObjectResource& InOutResource, bool bIsE
 			OuterString.RemoveSuffix(1);
 		}
 		SrcOuterName = FName(OuterString);
-		SrcObjectName = InOutResource.ObjectName;
+		SrcObjectName = InResource.ObjectName;
 	}
 	
-	const FCoreRedirectObjectName SrcObjectPath(SrcObjectName, SrcOuterName, SrcPackageName);
-	const FCoreRedirectObjectName DstObjectPath = FCoreRedirects::GetRedirectedName(ECoreRedirectFlags::Type_AllMask, SrcObjectPath);
-
-	bool bPatched = false;
-	if (!bRemapByPackageName)
-	{
-		bPatched = RemapFName(SrcObjectPath.ObjectName, DstObjectPath.ObjectName);
-		OutPatchedObjectName = DstObjectPath.ObjectName;
-	}
-	else
-	{
-		bPatched = RemapFName(SrcObjectPath.PackageName, DstObjectPath.PackageName);
-		OutPatchedObjectName = DstObjectPath.PackageName;
-	}
-
-#if WITH_EDITORONLY_DATA
-	InOutResource.OldClassName = NAME_None;
-#endif
-
-	return bPatched;
+	return FCoreRedirectObjectName(SrcObjectName, SrcOuterName, SrcPackageName);
 }
 
 bool FAssetHeaderPatcherInner::DoPatch(FTopLevelAssetPath& InOutPath)
@@ -1924,78 +2249,15 @@ void FAssetHeaderPatcherInner::PatchHeader_PatchSections()
 		}
 	}
 
-	// Patching of the FObjectResource ObjectNames is deferred since when patching we need to 
-	// walk the original names to determine if they need patching in the first place
-	TMap<int32, FName> PatchedExportObjectNames;
-	TMap<int32, FName> PatchedImportObjectNames;
-
-	// Export Table
-	{
-		PatchedExportObjectNames.Reserve(ExportTable.Num());
-		for (int32 i = 0; i < ExportTable.Num(); ++i)
-		{
-			FObjectExport& Export = ExportTable[i];
-			FName PatchedObjectName;
-			if (DoPatch(Export, true, PatchedObjectName))
-			{
-				PatchedExportObjectNames.Add(i, PatchedObjectName);
-			}
-		}
-	}
-
-	// Import table
-	{
-		PatchedImportObjectNames.Reserve(ImportTable.Num());
-		for (int32 i = 0; i < ImportTable.Num(); ++i)
-		{
-			FObjectImport& Import = ImportTable[i];
-			FName PatchedObjectName;
-			if (DoPatch(Import, false, PatchedObjectName))
-			{
-				PatchedImportObjectNames.Add(i, PatchedObjectName);
-			}
-
-			const FCoreRedirectObjectName SrcClass(Import.ClassName, NAME_None, Import.ClassPackage);
-			const FCoreRedirectObjectName DstClass = FCoreRedirects::GetRedirectedName(ECoreRedirectFlags::Type_Class | ECoreRedirectFlags::Type_Package, SrcClass);
-
-			if (RemapFName(SrcClass.ObjectName, DstClass.ObjectName))
-			{
-				Import.ClassName = DstClass.ObjectName;
-			}
-
-			if (RemapFName(SrcClass.PackageName, DstClass.PackageName))
-			{
-				Import.ClassPackage = DstClass.PackageName;
-			}
-
-#if WITH_EDITORONLY_DATA
-			const FCoreRedirectObjectName SrcPackage(NAME_None, NAME_None, Import.PackageName);
-			const FCoreRedirectObjectName DstPackage = FCoreRedirects::GetRedirectedName(ECoreRedirectFlags::Type_Package, SrcPackage);
-
-			if (RemapFName(SrcPackage.PackageName, DstPackage.PackageName))
-			{
-				Import.PackageName = DstPackage.PackageName;
-			}
-#endif
-		}
-	}
-
-	// Finish the FObjectResource patching that was deferred above for the ExportTable and ImportTable
-	{
-		for (auto& Pair : PatchedExportObjectNames)
-		{
-			int32 Index = Pair.Key;
-			FName PatchedName = Pair.Value;
-			ExportTable[Index].ObjectName = PatchedName;
-		}
-
-		for (auto& Pair : PatchedImportObjectNames)
-		{
-			int32 Index = Pair.Key;
-			FName PatchedName = Pair.Value;
-			ImportTable[Index].ObjectName = PatchedName;
-		}
-	}
+	// For Import and Export Tables we need to generate patches for both and apply them afterwards. This is due to the ObjectPaths
+	// being patched being split across multiple Export and Import entries that refer to one another. Patching them as we go would
+	// change the ObjectPaths of Exports/Imports we may need to patch but won't be able to deduce the original, unpatched ObjectPath.
+	TArray<FExportPatch> ExportPatches;
+	int32 NewImportCount = 0;
+	TArray<FImportPatch> ImportPatches;
+	GetExportTablePatches(ExportPatches);
+	GetImportTablePatches(ImportPatches, NewImportCount);
+	PatchExportAndImportTables(ExportPatches, ImportPatches, NewImportCount);
 
 	// Soft paths
 	for (FSoftObjectPath& SoftObjectPath : SoftObjectPathTable)
@@ -2364,8 +2626,8 @@ FAssetHeaderPatcher::EResult FAssetHeaderPatcherInner::PatchHeader_WriteDestinat
 
 			const int64 ImportTableSize = Writer.Tell() - ImportTableStartOffset;
 			const int64 Delta = ImportTableSize - SourceSection.Size;
-			check(Delta == 0);
-			checkf(ImportTableSize == SourceSection.Size, TEXT("%d == %d"), (int)ImportTableSize, (int)SourceSection.Size); // We only patch export table offsets, we should not be patching size
+			PatchSummaryOffsets(Summary, ImportTableStartOffset, Delta);
+			Summary.TotalHeaderSize += (int32)Delta;
 			checkf(Summary.ImportCount == ImportTable.Num(), TEXT("%d == %d"), Summary.ImportCount, ImportTable.Num());
 			checkf(Summary.ImportOffset == ImportTableStartOffset, TEXT("%d == %d"), Summary.ImportOffset, ImportTableStartOffset);
 
@@ -2551,7 +2813,8 @@ void FAssetHeaderPatcherInner::DumpState(FStringView OutputDirectory)
 			int32* Index = NameToIndexMap.Find(Name.GetDisplayIndex());
 			if (Index)
 			{
-				return FString::Printf(TEXT("%s (nametable index: %d, fname:{'%s', %d})"), *Name.ToString(), *Index, *Name.GetPlainNameString(), Name.GetNumber());
+				FName NameTableName = NameTable[*Index];
+				return FString::Printf(TEXT("%s (nametable index: %d, fname:{'%s', %d})"), *NameTableName.ToString(), *Index, *NameTableName.GetPlainNameString(), NameTableName.GetNumber());
 			}
 			else
 			{
@@ -2563,9 +2826,10 @@ void FAssetHeaderPatcherInner::DumpState(FStringView OutputDirectory)
 
 	Builder.Append(TEXT("\t\"Summary\":{ "));
 	{
-		Builder.Append(TEXT("\n\t\t\"PackageName\": \""));
-		Builder.Append(Summary.PackageName);
-		Builder.Append(TEXT("\""));
+		Builder.Appendf(TEXT("\n\t\t\"PackageName\": \"%s\""), *Summary.PackageName);
+		Builder.Appendf(TEXT(",\n\t\t\"NamesReferencedFromExportDataCount\": \"%d\""), Summary.NamesReferencedFromExportDataCount);
+		Builder.Appendf(TEXT(",\n\t\t\"ExportCount\": \"%d\""), Summary.ExportCount);
+		Builder.Appendf(TEXT(",\n\t\t\"ImportCount\": \"%d\""), Summary.ImportCount);
 	}
 	Builder.Append(TEXT("\n\t},\n"));
 
@@ -2580,16 +2844,79 @@ void FAssetHeaderPatcherInner::DumpState(FStringView OutputDirectory)
 	Builder.Append(TEXT("\n\t],\n"));
 
 	Builder.Append(TEXT("\t\"ExportTable\":[ "));
+	int ExportIndex = 0;
 	for (const auto& Export : ExportTable)
 	{
 		Builder.Append(TEXT("\n\t\t{\n"));
 
+		Builder.Appendf(TEXT("\t\t\t\"Index\": \"%d\""), ExportIndex++);
+		Builder.Append(TEXT("\",\n"));
+
 		Builder.Append(TEXT("\t\t\t\"ObjectName\": \""));
 		Builder.Append(GetDebugFNameString(Export.ObjectName));
-		Builder.Append(TEXT("\""));
+		Builder.Append(TEXT("\",\n"));
+
+		Builder.Append(TEXT("\t\t\t\"Outer\": \""));
+		if (Export.OuterIndex.IsNull())
+		{
+			Builder.Append(TEXT("None"));
+		}
+		else if (Export.OuterIndex.IsExport())
+		{
+			Builder.Appendf(TEXT("Export(%d) - %s"), Export.OuterIndex.ToExport(), *GetDebugFNameString(ExportTable[Export.OuterIndex.ToExport()].ObjectName));
+		}
+		else if (Export.OuterIndex.IsImport())
+		{
+			Builder.Appendf(TEXT("Import(%d) - %s"), Export.OuterIndex.ToImport(), *GetDebugFNameString(ImportTable[Export.OuterIndex.ToImport()].ObjectName));
+		}
+		Builder.Append(TEXT("\",\n"));
+
+		Builder.Append(TEXT("\t\t\t\"ClassIndex\": \""));
+		if (Export.ClassIndex.IsNull())
+		{
+			Builder.Append(TEXT("None"));
+		}
+		else if (Export.ClassIndex.IsExport())
+		{
+			Builder.Appendf(TEXT("Export(%d) - %s"), Export.ClassIndex.ToExport(), *GetDebugFNameString(ExportTable[Export.ClassIndex.ToExport()].ObjectName));
+		}
+		else if (Export.ClassIndex.IsImport())
+		{
+			Builder.Appendf(TEXT("Import(%d) - %s"), Export.ClassIndex.ToImport(), *GetDebugFNameString(ImportTable[Export.ClassIndex.ToImport()].ObjectName));
+		}
+		Builder.Append(TEXT("\",\n"));
+
+		Builder.Append(TEXT("\t\t\t\"SuperIndex\": \""));
+		if (Export.SuperIndex.IsNull())
+		{
+			Builder.Append(TEXT("None"));
+		}
+		else if (Export.SuperIndex.IsExport())
+		{
+			Builder.Appendf(TEXT("Export(%d) - %s"), Export.SuperIndex.ToExport(), *GetDebugFNameString(ExportTable[Export.SuperIndex.ToExport()].ObjectName));
+		}
+		else if (Export.SuperIndex.IsImport())
+		{
+			Builder.Appendf(TEXT("Import(%d) - %s"), Export.SuperIndex.ToImport(), *GetDebugFNameString(ImportTable[Export.SuperIndex.ToImport()].ObjectName));
+		}
+		Builder.Append(TEXT("\",\n"));
+
+		Builder.Append(TEXT("\t\t\t\"TemplateIndex\": \""));
+		if (Export.TemplateIndex.IsNull())
+		{
+			Builder.Append(TEXT("None"));
+		}
+		else if (Export.TemplateIndex.IsExport())
+		{
+			Builder.Appendf(TEXT("Export(%d) - %s"), Export.TemplateIndex.ToExport(), *GetDebugFNameString(ExportTable[Export.TemplateIndex.ToExport()].ObjectName));
+		}
+		else if (Export.TemplateIndex.IsImport())
+		{
+			Builder.Appendf(TEXT("Import(%d) - %s"), Export.TemplateIndex.ToImport(), *GetDebugFNameString(ImportTable[Export.TemplateIndex.ToImport()].ObjectName));
+		}
+		Builder.Append(TEXT("\",\n"));
 
 #if WITH_EDITORONLY_DATA
-		Builder.Append(TEXT(",\n"));
 		Builder.Append(TEXT("\t\t\t\"OldClassName\": \""));
 		Builder.Append(GetDebugFNameString(Export.OldClassName));
 		Builder.Append(TEXT("\""));
@@ -2600,14 +2927,33 @@ void FAssetHeaderPatcherInner::DumpState(FStringView OutputDirectory)
 	Builder.Append(TEXT("\n\t],\n"));
 
 	Builder.Append(TEXT("\t\"ImportTable\":[ "));
+	int ImportIndex = 0;
 	for (const auto& Import : ImportTable)
 	{
 		Builder.Append(TEXT("\n\t\t{\n"));
 
+		Builder.Appendf(TEXT("\t\t\t\"Index\": \"%d\""), ImportIndex++);
+		Builder.Append(TEXT("\",\n"));
+
 		Builder.Append(TEXT("\t\t\t\"ObjectName\": \""));
 		Builder.Append(GetDebugFNameString(Import.ObjectName));
 		Builder.Append(TEXT("\",\n"));
 
+		Builder.Append(TEXT("\t\t\t\"Outer\": \""));
+		if (Import.OuterIndex.IsNull())
+		{
+			Builder.Append(TEXT("None"));
+		}
+		else if (Import.OuterIndex.IsExport())
+		{
+			Builder.Appendf(TEXT("Export(%d) - %s"), Import.OuterIndex.ToExport(), *GetDebugFNameString(ExportTable[Import.OuterIndex.ToExport()].ObjectName));
+		}
+		else if (Import.OuterIndex.IsImport())
+		{
+			Builder.Appendf(TEXT("Import(%d) - %s"), Import.OuterIndex.ToImport(), *GetDebugFNameString(ImportTable[Import.OuterIndex.ToImport()].ObjectName));
+		}
+		Builder.Append(TEXT("\",\n"));
+
 #if WITH_EDITORONLY_DATA
 		Builder.Append(TEXT("\t\t\t\"OldClassName\": \""));
 		Builder.Append(GetDebugFNameString(Import.OldClassName));
@@ -2912,13 +3258,35 @@ TEST_CASE_NAMED(FAssetHeaderPatcherTests, "AssetHeaderPatcher", "[AssetHeaderPat
 	const TCHAR*  DstMountName				= TEXT("/DestinationSpecialMount/");
 	const FSoftObjectPath SoftObjectPathToRedirect(TEXT("/ToBeRedirectedMount/SomePath/ToBeRedirectedPackage.ToBeRedirectedPackage:Some.ToBeRedirectedPackage.Subobject"));
 	const FSoftObjectPath RedirectedSoftObjectPath(TEXT("/RedirectedMount/SomePath/RedirectedPackage.RedirectedPackage:Some.RedirectedPackage.Subobject"));
-	const FName   SrcPackagePathFName(SrcPackagePath);
-	const FName   DstPackagePathFName(DstPackagePath);
-	const FName   SrcAssetFName(TEXT("SrcPackage"));
-	const FName   DstAssetFName(TEXT("DstPackage"));
-	const FName   SrcExportObjectFName = SrcAssetFName;
-	const FName   DstExportObjectFName = DstAssetFName;
-	const FName   DummyImportPackagePathFName(TEXT("/DummyMount/DummyPackage"));
+	const FName SrcPackagePathFName(SrcPackagePath);
+	const FName DstPackagePathFName(DstPackagePath);
+	const FName SrcAssetFName(TEXT("SrcPackage"));
+	const FName DstAssetFName(TEXT("DstPackage"));
+	const FName SrcExportObjectFName = SrcAssetFName;
+	const FName DstExportObjectFName = DstAssetFName;
+	const FName DummyImportPackagePathFName(TEXT("/DummyMount/DummyPackage"));
+	// Import FNames
+	const FName SrcEngineModuleImportObjectName(TEXT("/Script/SrcEngineModule"));
+	const FName Default__SrcTypeAImportObjectName(TEXT("Default__SrcTypeA"));
+	const FName Default__SrcTypeBImportObjectName(TEXT("Default__SrcTypeB"));
+	const FName Default__OnlySubTypeChangedImportObjectName(TEXT("Default__OnlySubTypeChanged"));
+	const FName Default__MovedButNotRenamedTypeImportObjectName(TEXT("Default__MovedButNotRenamedType"));
+	const FName SrcPropertyToChangeAImportObjectName(TEXT("SrcPropertyToChangeA"));
+	const FName SrcPropertyToChangeBImportObjectName(TEXT("SrcPropertyToChangeB"));
+	const FName MovedButNotRenamedPropertyImportObjectName(TEXT("MovedButNotRenamedProperty"));
+	const FName InnerMovedButNotRenamedPropertyImportObjectName(TEXT("InnerMovedButNotRenamedProperty"));
+	const FName InnerInnerMovedButNotRenamedPropertyImportObjectName(TEXT("InnerInnerMovedButNotRenamedProperty"));
+	const FName UnchangedPropertyImportObjectName(TEXT("UnchangedProperty"));
+	const FName SrcImportClassName(TEXT("SrcClass"));
+	const FName SrcImportClassPackage(TEXT("/Engine/SrcClassPackage"));
+	
+	const FName DstEngineModuleImportObjectName(TEXT("/Script/DstEngineModule"));
+	const FName Default__DstTypeAImportObjectName(TEXT("Default__DstTypeA"));
+	const FName Default__DstTypeBImportObjectName(TEXT("Default__DstTypeB"));
+	const FName DstPropertyToChangeAImportObjectName(TEXT("DstPropertyToChangeA"));
+	const FName DstPropertyToChangeBImportObjectName(TEXT("DstPropertyToChangeB"));
+	const FName DstImportClassName(TEXT("DstClass"));
+	const FName DstImportClassPackage(TEXT("/Engine/DstClassPackage"));
 
 
 	TMap<FString, FString> MountPointReplacementMap =
@@ -2931,6 +3299,141 @@ TEST_CASE_NAMED(FAssetHeaderPatcherTests, "AssetHeaderPatcher", "[AssetHeaderPat
 		{ SrcPackagePath, DstPackagePath },
 	};
 
+	auto MakeImport = [SrcPackagePathFName](const FName ObjectName, const FPackageIndex OuterIndex, const FName ClassPackage, const FName ClassName, const FName PackageName)
+		{
+			FObjectImport Import;
+			Import.ObjectName = ObjectName;
+			Import.OuterIndex = OuterIndex;
+#if WITH_EDITORONLY_DATA
+			Import.OldClassName = NAME_None;
+			Import.PackageName = PackageName;
+#endif
+			Import.ClassPackage = ClassPackage;
+			Import.ClassName = ClassName;
+			return Import;
+		};
+
+	auto MakeExport = [](const FName ObjectName, const FPackageIndex ThisIndex, const FPackageIndex OuterIndex,
+		const FPackageIndex SuperIndex, const FPackageIndex ClassIndex, const FPackageIndex TemplateIndex)
+		{
+			FObjectExport Export;
+			Export.ObjectName = ObjectName;
+			Export.ThisIndex = ThisIndex;
+			Export.OuterIndex = OuterIndex;
+			Export.SuperIndex = SuperIndex;
+			Export.ClassIndex = ClassIndex;
+			Export.TemplateIndex = TemplateIndex;
+#if WITH_EDITORONLY_DATA
+			Export.OldClassName = NAME_None;
+#endif
+			return Export;
+		};
+
+	TArray<TPair<FObjectImport, FObjectImport>> ImportTestCases
+	{
+		// /Script/SrcEngineModule																-> (remains unchanged)
+		{	
+			MakeImport(SrcEngineModuleImportObjectName, FPackageIndex(), SrcImportClassPackage, SrcImportClassName, SrcPackagePathFName), 
+			MakeImport(SrcEngineModuleImportObjectName, FPackageIndex(), DstImportClassPackage, DstImportClassName, DstPackagePathFName) 
+		},
+		// /Script/SrcEngineModule.Default__SrcTypeA											-> /Script/DstEngineModule.Default__DstTypeA
+		{ 
+			MakeImport(Default__SrcTypeAImportObjectName, FPackageIndex::FromImport(0), SrcImportClassPackage, SrcImportClassName, SrcPackagePathFName), 
+			MakeImport(Default__DstTypeAImportObjectName, FPackageIndex::FromImport(13), DstImportClassPackage, DstImportClassName, DstPackagePathFName) 
+		},
+		// /Script/SrcEngineModule.Default__SrcTypeA.SrcPropertyToChangeA						-> /Script/DstEngineModule.Default__DstTypeA.DstPropertyToChangeA
+		{ 
+			MakeImport(SrcPropertyToChangeAImportObjectName, FPackageIndex::FromImport(1), SrcImportClassPackage, SrcImportClassName, SrcPackagePathFName), 
+			MakeImport(DstPropertyToChangeAImportObjectName, FPackageIndex::FromImport(1), DstImportClassPackage, DstImportClassName, DstPackagePathFName) 
+		},
+		// /Script/SrcEngineModule.Default__SrcTypeA.MovedButNotRenamedProperty					-> /Script/DstEngineModule.Default__DstTypeA.MovedButNotRenamedProperty
+		{ 
+			MakeImport(MovedButNotRenamedPropertyImportObjectName, FPackageIndex::FromImport(1), SrcImportClassPackage, SrcImportClassName, SrcPackagePathFName), 
+			MakeImport(MovedButNotRenamedPropertyImportObjectName, FPackageIndex::FromImport(1), DstImportClassPackage, DstImportClassName, DstPackagePathFName) 
+		},
+		// /Script/SrcEngineModule.Default__SrcTypeA.MovedButNotRenamedProperty.InnerMovedButNotRenamedProperty
+		// -> 
+		// /Script/DstEngineModule.Default__DstTypeA.MovedButNotRenamedProperty.InnerMovedButNotRenamedProperty
+		{
+			MakeImport(InnerMovedButNotRenamedPropertyImportObjectName, FPackageIndex::FromImport(3), SrcImportClassPackage, SrcImportClassName, SrcPackagePathFName),
+			MakeImport(InnerMovedButNotRenamedPropertyImportObjectName, FPackageIndex::FromImport(3), DstImportClassPackage, DstImportClassName, DstPackagePathFName)
+		},
+		// /Script/SrcEngineModule.Default__SrcTypeA.MovedButNotRenamedProperty.InnerMovedButNotRenamedProperty.InnerInnerMovedButNotRenamedProperty 
+		// -> 
+		// /Script/DstEngineModule.Default__DstTypeA.MovedButNotRenamedProperty.InnerMovedButNotRenamedProperty.InnerInnerMovedButNotRenamedProperty 
+		{
+			MakeImport(InnerInnerMovedButNotRenamedPropertyImportObjectName, FPackageIndex::FromImport(4), SrcImportClassPackage, SrcImportClassName, SrcPackagePathFName),
+			MakeImport(InnerInnerMovedButNotRenamedPropertyImportObjectName, FPackageIndex::FromImport(4), DstImportClassPackage, DstImportClassName, DstPackagePathFName)
+		},
+		// /Script/SrcEngineModule.Default__OnlySubTypeChanged									-> (remains unchanged)
+		{ 
+			MakeImport(Default__OnlySubTypeChangedImportObjectName, FPackageIndex::FromImport(0), SrcImportClassPackage, SrcImportClassName, SrcPackagePathFName), 
+			MakeImport(Default__OnlySubTypeChangedImportObjectName, FPackageIndex::FromImport(0), DstImportClassPackage, DstImportClassName, DstPackagePathFName)
+		},
+		// /Script/SrcEngineModule.Default__OnlySubTypeChanged.SrcPropertyToChangeB				-> /Script/SrcEngineModule.Default__OnlySubTypeChanged.DstPropertyToChangeB
+		{ 
+			MakeImport(SrcPropertyToChangeBImportObjectName, FPackageIndex::FromImport(6), SrcImportClassPackage, SrcImportClassName, SrcPackagePathFName), 
+			MakeImport(DstPropertyToChangeBImportObjectName, FPackageIndex::FromImport(6), DstImportClassPackage, DstImportClassName, DstPackagePathFName) 
+		},
+		// /Script/SrcEngineModule.Default__OnlySubTypeChanged.UnchangedProperty				-> (remains unchanged)
+		{ 
+			MakeImport(UnchangedPropertyImportObjectName, FPackageIndex::FromImport(6), SrcImportClassPackage, SrcImportClassName, SrcPackagePathFName), 
+			MakeImport(UnchangedPropertyImportObjectName, FPackageIndex::FromImport(6), DstImportClassPackage, DstImportClassName, DstPackagePathFName) 
+		},
+		// /Script/SrcEngineModule.Default__SrcTypeB											-> /Script/SrcEngineModule.Default__DstTypeB
+		{ 
+			MakeImport(Default__SrcTypeBImportObjectName, FPackageIndex::FromImport(0), SrcImportClassPackage, SrcImportClassName, SrcPackagePathFName), 
+			MakeImport(Default__DstTypeBImportObjectName, FPackageIndex::FromImport(0), DstImportClassPackage, DstImportClassName, DstPackagePathFName) 
+		},
+		// /Script/SrcEngineModule.Default__SrcTypeB.MovedButNotRenamedProperty					-> /Script/SrcEngineModule.Default__DstTypeB.MovedButNotRenamedProperty
+		{ 
+			MakeImport(MovedButNotRenamedPropertyImportObjectName, FPackageIndex::FromImport(9), SrcImportClassPackage, SrcImportClassName, SrcPackagePathFName), 
+			MakeImport(MovedButNotRenamedPropertyImportObjectName, FPackageIndex::FromImport(9), DstImportClassPackage, DstImportClassName, DstPackagePathFName) 
+		},
+		// /Script/SrcEngineModule.Default__MovedButNotRenamedType								-> /Script/DstEngineModule.Default__MovedButNotRenamedType
+		{ 
+			MakeImport(Default__MovedButNotRenamedTypeImportObjectName, FPackageIndex::FromImport(0), SrcImportClassPackage, SrcImportClassName, SrcPackagePathFName), 
+			MakeImport(Default__MovedButNotRenamedTypeImportObjectName, FPackageIndex::FromImport(13), DstImportClassPackage, DstImportClassName, DstPackagePathFName) 
+		},
+		// /Script/SrcEngineModule.Default__MovedButNotRenamedType.MovedButNotRenamedProperty	-> /Script/DstEngineModule.Default__MovedButNotRenamedType.MovedButNotRenamedProperty
+		{ 
+			MakeImport(MovedButNotRenamedPropertyImportObjectName, FPackageIndex::FromImport(11), SrcImportClassPackage, SrcImportClassName, SrcPackagePathFName), 
+			MakeImport(MovedButNotRenamedPropertyImportObjectName, FPackageIndex::FromImport(11), DstImportClassPackage, DstImportClassName, DstPackagePathFName) 
+		},
+		// <none>																				-> /Script/DstEngineModule
+		{ 
+			FObjectImport(), 
+			MakeImport(DstEngineModuleImportObjectName, FPackageIndex(), DstImportClassPackage, DstImportClassName, DstPackagePathFName)
+		},
+	};
+
+	TArray<TPair<FObjectExport, FObjectExport>> ExportTestCases
+	{
+		// SrcPackage																			-> DstPackage
+		// Super: /Script/SrcEngineModule														-> /Script/DstEngineModule
+		// Class: /Script/SrcEngineModule.Default__SrcTypeA.SrcPropertyToChangeA				-> /Script/DstEngineModule.Default__DstTypeA.DstPropertyToChangeA
+		// Template: /Script/SrcEngineModule.Default__OnlySubTypeChanged						-> /Script/SrcEngineModule.Default__OnlySubTypeChanged
+		{
+			MakeExport(SrcExportObjectFName, FPackageIndex::FromExport(0), FPackageIndex(), FPackageIndex::FromImport(0), FPackageIndex::FromImport(2), FPackageIndex::FromImport(4)),
+			MakeExport(DstExportObjectFName, FPackageIndex::FromExport(0), FPackageIndex(), FPackageIndex::FromImport(11), FPackageIndex::FromImport(2), FPackageIndex::FromImport(4))
+		},
+		// SrcPackage.SrcPackage																-> DstPackage.DstPackage
+		// Super: /Script/SrcEngineModule.Default__SrcTypeA										-> /Script/DstEngineModule.Default__DstTypeA
+		// Class: /Script/SrcEngineModule.Default__SrcTypeA.MovedButNotRenamedProperty			-> /Script/DstEngineModule.Default__DstTypeA.MovedButNotRenamedProperty
+		// Template: /Script/SrcEngineModule.Default__OnlySubTypeChanged.SrcPropertyToChangeB	-> /Script/SrcEngineModule.Default__OnlySubTypeChanged.DstPropertyToChangeB
+		{
+			MakeExport(SrcExportObjectFName, FPackageIndex::FromExport(1), FPackageIndex::FromExport(0), FPackageIndex::FromImport(1), FPackageIndex::FromImport(3), FPackageIndex::FromImport(5)),
+			MakeExport(DstExportObjectFName, FPackageIndex::FromExport(1), FPackageIndex::FromExport(0), FPackageIndex::FromImport(1), FPackageIndex::FromImport(3), FPackageIndex::FromImport(5))
+		}
+	};
+
+	FCoreRedirectsContext TestRedirectContext;
+	TestRedirectContext.InitializeContext();
+	CHECK(TestRedirectContext.IsInitialized());
+	FCoreRedirectsContext& OriginalContext = FCoreRedirectsContext::GetThreadContext();
+	FCoreRedirectsContext::SetThreadContext(TestRedirectContext);
+	ON_SCOPE_EXIT{ FCoreRedirectsContext::SetThreadContext(OriginalContext); };
+
 	FTestPatcherContext Context(PackageRenameMap, false /*bGatherDependentPackages*/);
 	const TMap<FString, FString>& StringReplacements = Context.GetStringReplacements();
 	CHECK(StringReplacements.Num() > PackageRenameMap.Num()); // Ensure we generated more mappings off of the PackageRenameMap
@@ -2939,16 +3442,14 @@ TEST_CASE_NAMED(FAssetHeaderPatcherTests, "AssetHeaderPatcher", "[AssetHeaderPat
 	FAssetHeaderPatcherInner Patcher(DummySrcDstAsset, DummySrcDstAsset, StringReplacements, MountPointReplacementMap);
 
 	int32 OriginalNameTableCount = 0;
-	auto ResetPatcher = [&Patcher, &OriginalNameTableCount, SrcPackagePathFName, SrcAssetFName, SrcExportObjectFName, &SoftObjectPathToRedirect, DummyImportPackagePathFName]()
+	auto ResetPatcher = [&]()
 		{
-			// Reset NameTable
-			Patcher.NameTable.Empty();
-			Patcher.NameToIndexMap.Empty();
-			Patcher.RenameMap.Empty();
-			Patcher.AddedNames.Empty();
-			Patcher.ExportTable.Empty();
+			// Reset Patcher
+			Patcher.ResetInternalState();
 
-			// Repopulate it with our test data normally set through deserialization
+			// Repopulate patcher state with our test data normally set through deserialization
+			///////////////////////////////////////////////////////////////////////////////////
+			
 			// NameTable
 			Patcher.NameToIndexMap.Add(SrcPackagePathFName.GetDisplayIndex(), Patcher.NameTable.Num());
 			Patcher.NameTable.Add(SrcPackagePathFName);
@@ -2965,25 +3466,72 @@ TEST_CASE_NAMED(FAssetHeaderPatcherTests, "AssetHeaderPatcher", "[AssetHeaderPat
 			Patcher.NameTable.Add(ToBeRedirectedPackageName);
 			Patcher.NameToIndexMap.Add(ToBeRedirectedAssetName.GetDisplayIndex(), Patcher.NameTable.Num());
 			Patcher.NameTable.Add(ToBeRedirectedAssetName);
+			// Export Table Names
+			Patcher.NameToIndexMap.Add(SrcExportObjectFName.GetDisplayIndex(), Patcher.NameTable.Num());
+			Patcher.NameTable.Add(SrcExportObjectFName);
+			// Import Table Names
+			Patcher.NameToIndexMap.Add(SrcEngineModuleImportObjectName.GetDisplayIndex(), Patcher.NameTable.Num());
+			Patcher.NameTable.Add(SrcEngineModuleImportObjectName);
+			Patcher.NameToIndexMap.Add(Default__SrcTypeAImportObjectName.GetDisplayIndex(), Patcher.NameTable.Num());
+			Patcher.NameTable.Add(Default__SrcTypeAImportObjectName);
+			Patcher.NameToIndexMap.Add(Default__SrcTypeBImportObjectName.GetDisplayIndex(), Patcher.NameTable.Num());
+			Patcher.NameTable.Add(Default__SrcTypeBImportObjectName);
+			Patcher.NameToIndexMap.Add(Default__OnlySubTypeChangedImportObjectName.GetDisplayIndex(), Patcher.NameTable.Num());
+			Patcher.NameTable.Add(Default__OnlySubTypeChangedImportObjectName);
+			Patcher.NameToIndexMap.Add(Default__MovedButNotRenamedTypeImportObjectName.GetDisplayIndex(), Patcher.NameTable.Num());
+			Patcher.NameTable.Add(Default__MovedButNotRenamedTypeImportObjectName);
+			Patcher.NameToIndexMap.Add(SrcPropertyToChangeAImportObjectName.GetDisplayIndex(), Patcher.NameTable.Num());
+			Patcher.NameTable.Add(SrcPropertyToChangeAImportObjectName);
+			Patcher.NameToIndexMap.Add(SrcPropertyToChangeBImportObjectName.GetDisplayIndex(), Patcher.NameTable.Num());
+			Patcher.NameTable.Add(SrcPropertyToChangeBImportObjectName);
+			Patcher.NameToIndexMap.Add(MovedButNotRenamedPropertyImportObjectName.GetDisplayIndex(), Patcher.NameTable.Num());
+			Patcher.NameTable.Add(MovedButNotRenamedPropertyImportObjectName);
+			Patcher.NameToIndexMap.Add(InnerMovedButNotRenamedPropertyImportObjectName.GetDisplayIndex(), Patcher.NameTable.Num());
+			Patcher.NameTable.Add(InnerMovedButNotRenamedPropertyImportObjectName);
+			Patcher.NameToIndexMap.Add(InnerInnerMovedButNotRenamedPropertyImportObjectName.GetDisplayIndex(), Patcher.NameTable.Num());
+			Patcher.NameTable.Add(InnerInnerMovedButNotRenamedPropertyImportObjectName);
+			Patcher.NameToIndexMap.Add(UnchangedPropertyImportObjectName.GetDisplayIndex(), Patcher.NameTable.Num());
+			Patcher.NameTable.Add(UnchangedPropertyImportObjectName);
+			Patcher.NameToIndexMap.Add(SrcImportClassName.GetDisplayIndex(), Patcher.NameTable.Num());
+			Patcher.NameTable.Add(SrcImportClassName);
+			Patcher.NameToIndexMap.Add(SrcImportClassPackage.GetDisplayIndex(), Patcher.NameTable.Num());
+			Patcher.NameTable.Add(SrcImportClassPackage);
+
+			auto CheckNameTableInit = [&Patcher](FName Name)
+				{
+					CHECK(Patcher.NameToIndexMap.Contains(Name.GetDisplayIndex()));
+					CHECK(Patcher.NameTable[Patcher.NameToIndexMap[Name.GetDisplayIndex()]] == Name);
+				};
 
-			FObjectImport DummyImport;
-			DummyImport.ObjectName = Patcher.NameTable[Patcher.NameToIndexMap[DummyImportPackagePathFName.GetDisplayIndex()]];
-			DummyImport.OldClassName = DummyImportPackagePathFName; // Set to something other than NAME_None
-			DummyImport.OuterIndex = FPackageIndex();
-			Patcher.ImportTable.Add(DummyImport);
-
-			// Export Table
-			FObjectExport SrcPackageExport;
-			SrcPackageExport.ObjectName = Patcher.NameTable[Patcher.NameToIndexMap[SrcExportObjectFName.GetDisplayIndex()]];
-			SrcPackageExport.OldClassName = SrcExportObjectFName; // Set to something other than NAME_None
-			SrcPackageExport.OuterIndex = FPackageIndex(); // This package is the outer
-			Patcher.ExportTable.Add(SrcPackageExport);
-
-			FObjectExport DummyExport;
-			DummyExport.ObjectName = Patcher.NameTable[Patcher.NameToIndexMap[SrcExportObjectFName.GetDisplayIndex()]]; // Same name as SrcPackageExport
-			DummyExport.OldClassName = SrcExportObjectFName; // Set to something other than NAME_None
-			DummyExport.OuterIndex = FPackageIndex::FromImport(0); // DummyImport is our outer
-			Patcher.ExportTable.Add(DummyExport);
+			// Import/Export Table (see breakdown of import names for the intended tests in the FObjectResource test section below)
+			for (const TPair<FObjectImport, FObjectImport>& TestCase : ImportTestCases)
+			{
+				const FObjectImport& Import = TestCase.Key;
+				if (Import.ObjectName == NAME_None)
+				{
+					// We have more cases than initial starting states so if we see an empty
+					// name we can stop adding to the initial state
+					break;
+				}
+				CheckNameTableInit(Import.ObjectName);
+				CheckNameTableInit(Import.ClassName);
+				CheckNameTableInit(Import.ClassPackage);
+#if WITH_EDITORONLY_DATA
+				CheckNameTableInit(Import.PackageName);
+				CheckNameTableInit(Import.OldClassName);
+#endif
+				Patcher.ImportTable.Add(Import);
+			}
+
+			for (const TPair<FObjectExport, FObjectExport>& TestCase : ExportTestCases)
+			{
+				const FObjectExport& Export = TestCase.Key;
+				CheckNameTableInit(Export.ObjectName);
+#if WITH_EDITORONLY_DATA
+				CheckNameTableInit(Export.OldClassName);
+#endif
+				Patcher.ExportTable.Add(Export);
+			}		
 
 			// Summary
 			Patcher.Summary.NameCount = Patcher.NameTable.Num();
@@ -3454,48 +4002,138 @@ TEST_CASE_NAMED(FAssetHeaderPatcherTests, "AssetHeaderPatcher", "[AssetHeaderPat
 		}
 	}
 
-	SECTION("DoPatch(FObjectResource)")
+	// Import Test Breakdown
+	// 
+	// Assume that SrcEngineModule and DstEngineModule both exist and types are moving between them and/or possibly being renamed while moving
+	// 
+	// Coreredirects are added to make the following transformations:
+	// 
+	// /Script/SrcEngineModule.Default__SrcTypeA.SrcPropertyToChangeA                     -> /Script/DstEngineModule.Default__DstTypeA.DstPropertyToChangeA
+	// /Script/SrcEngineModule.Default__SrcTypeA.MovedButNotRenamedProperty               -> /Script/DstEngineModule.Default__DstTypeA.MovedButNotRenamedProperty
+	// /Script/SrcEngineModule.Default__OnlySubTypeChanged.SrcPropertyToChangeB           -> /Script/SrcEngineModule.Default__OnlySubTypeChanged.DstPropertyToChangeB
+	// /Script/SrcEngineModule.Default__SrcTypeB.MovedButNotRenamedProperty               -> /Script/SrcEngineModule.Default__DstTypeB.MovedButNotRenamedProperty
+	// /Script/SrcEngineModule.Default__MovedButNotRenamedType.MovedButNotRenamedProperty -> /Script/DstEngineModule.Default__MovedButNotRenamedType.MovedButNotRenamedProperty
+	// 
+	// Our ImportTable will be transformed from the left side to the right where each line is its own ImportEntry outered to the import "up and to the left" of it:
+	// 
+	// /Script/SrcEngineModule                                                                          /Script/SrcEngineModule                                                           
+	//                        .Default__SrcTypeA										                                       .Default__OnlySubTypeChanged							     
+	//                                         .SrcPropertyToChangeA					                                                                   .DstPropertyToChangeB		     
+	//						  				   .MovedButNotRenamedProperty			          	         					                               .UnchangedProperty    
+	//					      .Default__OnlySubTypeChanged							          	         				       .Default__DstTypeB										   
+	//                                                    .SrcPropertyToChangeB		          	                                                 .MovedButNotRenamedProperty	          
+	//					                                  .UnchangedProperty                  -->       /Script/DstEngineModule				       
+	//					      .Default__SrcTypeB										                				       .Default__DstTypeA
+	//					                        .MovedButNotRenamedProperty	          			         				                         .DstPropertyToChangeA
+	//					      .Default__MovedButNotRenamedType						          	         				                         .MovedButNotRenamedProperty
+	//					                                      .MovedButNotRenamedProperty               				       .Default__MovedButNotRenamedType						     
+	//                                                                                                                                                         .MovedButNotRenamedProperty
+
+	SECTION("FObjectResource Patching")
 	{
+		SECTION("Redirect object to new package keeps the original package name if still in use")
 		{
+			auto CheckFNames = [&Patcher](FName Expected, FName Actual)
+				{
+					CHECK(Expected == Actual);
+					// It's fine to compare Expected == Actual when None but the NameTable
+					// should not contain None, so don't look for it
+					if (!Actual.IsNone())
+					{
+						CHECK(Patcher.NameTable.Contains(Actual));
+					}
+				};
+
 			ResetPatcher();
+			TArray<FCoreRedirect> ImportTableRedirects =
+			{
+				{ ECoreRedirectFlags::Type_Package, TEXT("/Engine/SrcClassPackage"),																		TEXT("/Engine/DstClassPackage")},
+				{ ECoreRedirectFlags::Type_Class, TEXT("/Engine/SrcClassPackage.SrcClass"),																	TEXT("/Engine/DstClassPackage.DstClass")},
+				{ ECoreRedirectFlags::Type_Class, TEXT("/Script/SrcEngineModule.Default__SrcTypeA"),														TEXT("/Script/DstEngineModule.Default__DstTypeA")},
+				{ ECoreRedirectFlags::Type_Class, TEXT("/Script/SrcEngineModule.Default__SrcTypeA.SrcPropertyToChangeA"),									TEXT("/Script/DstEngineModule.Default__DstTypeA.DstPropertyToChangeA")},
+				{ ECoreRedirectFlags::Type_Class, TEXT("/Script/SrcEngineModule.Default__SrcTypeA.MovedButNotRenamedProperty"),								TEXT("/Script/DstEngineModule.Default__DstTypeA.MovedButNotRenamedProperty")},
+				{ ECoreRedirectFlags::Type_Class, TEXT("/Script/SrcEngineModule.Default__OnlySubTypeChanged.SrcPropertyToChangeB"),							TEXT("/Script/SrcEngineModule.Default__OnlySubTypeChanged.DstPropertyToChangeB")},
+				{ ECoreRedirectFlags::Type_Class, TEXT("/Script/SrcEngineModule.Default__SrcTypeB"),														TEXT("/Script/SrcEngineModule.Default__DstTypeB")},
+				{ ECoreRedirectFlags::Type_Class, TEXT("/Script/SrcEngineModule.Default__SrcTypeB.MovedButNotRenamedProperty"),								TEXT("/Script/SrcEngineModule.Default__DstTypeB.MovedButNotRenamedProperty")},
+				{ ECoreRedirectFlags::Type_Class, TEXT("/Script/SrcEngineModule.Default__MovedButNotRenamedType"),											TEXT("/Script/DstEngineModule.Default__MovedButNotRenamedType")},
+				{ ECoreRedirectFlags::Type_Class, TEXT("/Script/SrcEngineModule.Default__MovedButNotRenamedType.MovedButNotRenamedProperty"),				TEXT("/Script/DstEngineModule.Default__MovedButNotRenamedType.MovedButNotRenamedProperty")}
+			};
+			CHECK(FCoreRedirects::AddRedirectList(ImportTableRedirects, TEXT("Asset Header Patcher Tests - FObjectResource Patching")));
 
-			CHECK(Patcher.ExportTable.Num() > 0);
-			FObjectResource ExportResource = Patcher.ExportTable[0];
-			CHECK(ExportResource.OldClassName != NAME_None);
+			// Confirm the initial state before patching is what we expect
+			check(Patcher.ImportTable.Num() <= ImportTestCases.Num());
+			for (int32 i = 0; i < Patcher.ImportTable.Num(); ++i)
+			{
+				const FObjectImport& Expected = ImportTestCases[i].Key;
+				const FObjectImport& Actual = Patcher.ImportTable[i];
 
-			FName Actual = ExportResource.ObjectName;
-			const FName Expected = DstExportObjectFName;
-			CHECK(Actual == SrcExportObjectFName);
-			CHECK(Patcher.NameTable.Contains(Actual));
+				CheckFNames(Expected.ObjectName, Actual.ObjectName);
+				CHECK(Expected.OuterIndex == Actual.OuterIndex);
+				CheckFNames(Expected.ClassName, Actual.ClassName);
+				CheckFNames(Expected.ClassPackage, Actual.ClassPackage);
+#if WITH_EDITORONLY_DATA
+				//CheckFNames(Expected.PackageName, Actual.PackageName);
+#endif
+			}
+			check(Patcher.ExportTable.Num() <= ExportTestCases.Num());
+			for (int32 i = 0; i < Patcher.ExportTable.Num(); ++i)
+			{
+				const FObjectExport& Expected = ExportTestCases[i].Key;
+				const FObjectExport& Actual = Patcher.ExportTable[i];
 
-			CHECK(Patcher.DoPatch(ExportResource, true, Actual));
-			CHECK(ExportResource.OldClassName == NAME_None);
-			CHECK(Actual == Expected);
-		}
+				CheckFNames(Expected.ObjectName, Actual.ObjectName);
+				CHECK(Expected.OuterIndex == Actual.OuterIndex);
+#if WITH_EDITORONLY_DATA
+				CheckFNames(Expected.OldClassName, Actual.OldClassName);
+				CheckFNames(Actual.OldClassName, NAME_None);
+#endif
+			}
 
-		// We are looking at an Export object that has the same name as an export that we will patch
-		// however, in this instance the outers are different. This export should not be patched as a 
-		// result since we do not have a remapping for it's full object path
-		{
-			ResetPatcher();
+			// Perform patching
+
+			TArray<FAssetHeaderPatcherInner::FExportPatch> ExportPatches;
+			int32 NewImportCount = 0;
+			TArray<FAssetHeaderPatcherInner::FImportPatch> ImportPatches;
+			Patcher.GetExportTablePatches(ExportPatches);
+			CHECK(!ExportPatches.IsEmpty());
+			Patcher.GetImportTablePatches(ImportPatches, NewImportCount);
+			CHECK(!ExportPatches.IsEmpty());
+			Patcher.PatchExportAndImportTables(ExportPatches, ImportPatches, NewImportCount);
+			Patcher.PatchNameTable();
 
-			CHECK(Patcher.ExportTable.Num() > 1);
-			FObjectResource ExportResource = Patcher.ExportTable[1];
-			CHECK(ExportResource.OldClassName != NAME_None);
-			CHECK(ExportResource.OuterIndex.IsImport());
+			// Confirm the patched state is what is expected
+			CHECK(Patcher.ImportTable.Num() == ImportTestCases.Num());
+			for (int32 i = 0; i < ImportTestCases.Num(); ++i)
+			{
+				const FObjectImport& Expected = ImportTestCases[i].Value;
+				const FObjectImport& Actual = Patcher.ImportTable[i];
 
-			FName Actual = ExportResource.ObjectName;
-			const FName Expected = Actual;
-			CHECK(Actual == SrcExportObjectFName);
-			CHECK(Patcher.NameTable.Contains(Actual));
+				CheckFNames(Expected.ObjectName, Actual.ObjectName);
+				CHECK(Expected.OuterIndex == Actual.OuterIndex);
+				CheckFNames(Expected.ClassName, Actual.ClassName);
+				CheckFNames(Expected.ClassPackage, Actual.ClassPackage);
+#if WITH_EDITORONLY_DATA
+				CheckFNames(Expected.PackageName, Actual.PackageName);
+				CheckFNames(Expected.OldClassName, Actual.OldClassName);
+				CheckFNames(Actual.OldClassName, NAME_None);
+#endif
+			}
+			
+			CHECK(Patcher.ExportTable.Num() == ExportTestCases.Num());
+			for (int32 i = 0; i < Patcher.ExportTable.Num(); ++i)
+			{
+				const FObjectExport& Expected = ExportTestCases[i].Value;
+				const FObjectExport& Actual = Patcher.ExportTable[i];
 
-			CHECK(!Patcher.DoPatch(ExportResource, true, Actual));
-			CHECK(ExportResource.OldClassName == NAME_None); // We always clear this
-			CHECK(Actual == Expected);
+				CheckFNames(Expected.ObjectName, Actual.ObjectName);
+				CHECK(Expected.OuterIndex == Actual.OuterIndex);
+#if WITH_EDITORONLY_DATA
+				CheckFNames(Expected.OldClassName, Actual.OldClassName);
+				CheckFNames(Actual.OldClassName, NAME_None);
+#endif
+			}
 		}
 	}
-
-	CHECK(FCoreRedirects::RemoveRedirectList(Context.GetRedirects(), TEXT("Asset Header Patcher Tests")));
 }
 
 #endif // WITH_TESTS
diff --git a/Engine/Source/Runtime/CoreUObject/Public/UObject/CoreRedirects.h b/Engine/Source/Runtime/CoreUObject/Public/UObject/CoreRedirects.h
--- a/Engine/Source/Runtime/CoreUObject/Public/UObject/CoreRedirects.h
+++ b/Engine/Source/Runtime/CoreUObject/Public/UObject/CoreRedirects.h
@@ -182,6 +182,11 @@ struct FCoreRedirectObjectName
 		return ObjectName != NAME_None || PackageName != NAME_None;
 	}
 
+	friend uint32 GetTypeHash(const FCoreRedirectObjectName& RedirectName)
+	{
+		return HashCombine(GetTypeHash(RedirectName.ObjectName), HashCombine(GetTypeHash(RedirectName.OuterName), GetTypeHash(RedirectName.PackageName)));
+	}
+
 	/** Returns true if all names have valid characters */
 	COREUOBJECT_API bool HasValidCharacters(ECoreRedirectFlags Type) const;
 
