Exposed utils for querying whether and asset or object is cooked or has its editor-only data, or whether an asset is considered read-only in-memory

diff --git a/Engine/Source/Developer/AssetTools/Private/AssetTools.cpp b/Engine/Source/Developer/AssetTools/Private/AssetTools.cpp
--- a/Engine/Source/Developer/AssetTools/Private/AssetTools.cpp
+++ b/Engine/Source/Developer/AssetTools/Private/AssetTools.cpp
@@ -6275,6 +6275,23 @@ bool UAssetToolsImpl::IsAssetVisible(const FAssetData& AssetData, bool bCheckAli
 	return false;
 }
 
+bool UAssetToolsImpl::IsAssetReadOnly(const FAssetData& AssetData) const
+{
+	// Cooked assets are always read-only
+	if (AssetData.HasAnyPackageFlags(PKG_Cooked | PKG_FilterEditorOnly))
+	{
+		return true;
+	}
+
+	// If this asset's package isn't in a writable folder, then it's read-only
+	if (!WritableFolderPermissionList->PassesStartsWithFilter(AssetData.PackagePath))
+	{
+		return true;
+	}
+
+	return false;
+}
+
 bool UAssetToolsImpl::AllPassWritableFolderFilter(const TArray<FString>& InPaths) const
 {
 	if (WritableFolderPermissionList->HasFiltering())
diff --git a/Engine/Source/Developer/AssetTools/Private/AssetTools.h b/Engine/Source/Developer/AssetTools/Private/AssetTools.h
--- a/Engine/Source/Developer/AssetTools/Private/AssetTools.h
+++ b/Engine/Source/Developer/AssetTools/Private/AssetTools.h
@@ -145,6 +145,7 @@ public:
 	virtual TSharedRef<FPathPermissionList>& GetFolderPermissionList() override;
 	virtual TSharedRef<FPathPermissionList>& GetWritableFolderPermissionList() override;
 	virtual bool IsAssetVisible(const FAssetData& AssetData, bool bCheckAliases = true) const override;
+	virtual bool IsAssetReadOnly(const FAssetData& AssetData) const override;
 	virtual bool AllPassWritableFolderFilter(const TArray<FString>& InPaths) const override;
 	virtual void NotifyBlockedByWritableFolderFilter() const;
 	virtual bool IsNameAllowed(const FString& Name, FText* OutErrorMessage = nullptr) const override;
diff --git a/Engine/Source/Developer/AssetTools/Public/IAssetTools.h b/Engine/Source/Developer/AssetTools/Public/IAssetTools.h
--- a/Engine/Source/Developer/AssetTools/Public/IAssetTools.h
+++ b/Engine/Source/Developer/AssetTools/Public/IAssetTools.h
@@ -649,6 +649,13 @@ public:
 	/** Determines whether an asset has a viewable asset class and its folder is allowed. If not, it can also check if it's aliased to an allowed location. */
 	virtual bool IsAssetVisible(const FAssetData& AssetData, bool bCheckAliases = true) const = 0;
 
+	/**
+	 * Determines whether an asset is considered read-only for editing, based on its flags and the current permissions (@see GetWritableFolderPermissionList).
+	 * @note This does not check the read-only state of the file itself or anything regarding source control, as it's only checking if the asset can be edited in-memory, not that any edits can be saved to disk.
+	 */
+	UFUNCTION(BlueprintCallable, Category = "Editor Scripting | Asset Tools")
+	virtual bool IsAssetReadOnly(const FAssetData& AssetData) const = 0;
+
 	/** Returns true if all in list pass writable folder filter */
 	virtual bool AllPassWritableFolderFilter(const TArray<FString>& InPaths) const = 0;
 
diff --git a/Engine/Source/Runtime/AssetRegistry/Private/AssetRegistryHelpers.cpp b/Engine/Source/Runtime/AssetRegistry/Private/AssetRegistryHelpers.cpp
--- a/Engine/Source/Runtime/AssetRegistry/Private/AssetRegistryHelpers.cpp
+++ b/Engine/Source/Runtime/AssetRegistry/Private/AssetRegistryHelpers.cpp
@@ -92,6 +92,18 @@ bool UAssetRegistryHelpers::IsAssetLoaded(const FAssetData& InAssetData)
 	return InAssetData.IsAssetLoaded();
 }
 
+#if WITH_EDITOR
+bool UAssetRegistryHelpers::IsAssetCooked(const FAssetData& InAssetData)
+{
+	return InAssetData.HasAnyPackageFlags(PKG_Cooked);
+}
+
+bool UAssetRegistryHelpers::AssetHasEditorOnlyData(const FAssetData& InAssetData)
+{
+	return !InAssetData.HasAnyPackageFlags(PKG_FilterEditorOnly);
+}
+#endif	// WITH_EDITOR
+
 FString UAssetRegistryHelpers::GetExportTextName(const FAssetData& InAssetData)
 {
 	return InAssetData.GetExportTextName();
diff --git a/Engine/Source/Runtime/AssetRegistry/Public/AssetRegistry/AssetRegistryHelpers.h b/Engine/Source/Runtime/AssetRegistry/Public/AssetRegistry/AssetRegistryHelpers.h
--- a/Engine/Source/Runtime/AssetRegistry/Public/AssetRegistry/AssetRegistryHelpers.h
+++ b/Engine/Source/Runtime/AssetRegistry/Public/AssetRegistry/AssetRegistryHelpers.h
@@ -91,6 +91,16 @@ public:
 	UFUNCTION(BlueprintPure, Category = "Asset Data", meta=(ScriptMethod))
 	static bool IsAssetLoaded(const FAssetData& InAssetData);
 
+#if WITH_EDITOR
+	/** Returns true if the asset is cooked */
+	UFUNCTION(BlueprintPure, Category = "Asset Data", meta=(ScriptMethod="IsCooked"))
+	static bool IsAssetCooked(const FAssetData& InAssetData);
+
+	/** Returns true if the asset has its editor-only data */
+	UFUNCTION(BlueprintPure, Category = "Asset Data", meta=(ScriptMethod="HasEditorOnlyData"))
+	static bool AssetHasEditorOnlyData(const FAssetData& InAssetData);
+#endif	// WITH_EDITOR
+
 	/** Returns the name for the asset in the form: Class'ObjectPath' */
 	UFUNCTION(BlueprintPure, Category = "Asset Data", meta=(ScriptMethod))
 	static FString GetExportTextName(const FAssetData& InAssetData);
diff --git a/Engine/Source/Runtime/Engine/Classes/Kismet/KismetSystemLibrary.h b/Engine/Source/Runtime/Engine/Classes/Kismet/KismetSystemLibrary.h
--- a/Engine/Source/Runtime/Engine/Classes/Kismet/KismetSystemLibrary.h
+++ b/Engine/Source/Runtime/Engine/Classes/Kismet/KismetSystemLibrary.h
@@ -175,6 +175,14 @@ class UKismetSystemLibrary : public UBlueprintFunctionLibrary
 	static ENGINE_API UObject* GetOuterObject(const UObject* Object);
 
 #if WITH_EDITOR
+	/** Returns true if the object is cooked */
+	UFUNCTION(BlueprintPure, Category = "Utilities")
+	static ENGINE_API bool IsObjectCooked(const UObject* Object);
+
+	/** Returns true if the object has its editor-only data */
+	UFUNCTION(BlueprintPure, Category = "Utilities")
+	static ENGINE_API bool ObjectHasEditorOnlyData(const UObject* Object);
+
 	/**
 	 * Duplicate an object to a new instance
 	 *
diff --git a/Engine/Source/Runtime/Engine/Private/KismetSystemLibrary.cpp b/Engine/Source/Runtime/Engine/Private/KismetSystemLibrary.cpp
--- a/Engine/Source/Runtime/Engine/Private/KismetSystemLibrary.cpp
+++ b/Engine/Source/Runtime/Engine/Private/KismetSystemLibrary.cpp
@@ -217,6 +217,16 @@ UObject* UKismetSystemLibrary::GetOuterObject(const UObject* Object)
 }
 
 #if WITH_EDITOR
+bool UKismetSystemLibrary::IsObjectCooked(const UObject* Object)
+{
+	return Object && Object->GetPackage()->HasAnyPackageFlags(PKG_Cooked);
+}
+
+bool UKismetSystemLibrary::ObjectHasEditorOnlyData(const UObject* Object)
+{
+	return Object && !Object->GetPackage()->HasAnyPackageFlags(PKG_FilterEditorOnly);
+}
+
 UObject* UKismetSystemLibrary::DuplicateObject(const UObject* SourceObject, UObject* Outer, const FName Name)
 {
 	return SourceObject ? ::DuplicateObject(SourceObject, Outer, Name) : nullptr;
