SourceControl / AssetRenameManager: Prevert revision control prompts from showing up if revision control is disabled. The prompt can be confusing to people, giving the suggestion that revision control operations are happening.

#rnx

diff --git a/Engine/Source/Developer/AssetTools/Private/AssetRenameManager.cpp b/Engine/Source/Developer/AssetTools/Private/AssetRenameManager.cpp
--- a/Engine/Source/Developer/AssetTools/Private/AssetRenameManager.cpp
+++ b/Engine/Source/Developer/AssetTools/Private/AssetRenameManager.cpp
@@ -481,7 +481,11 @@ bool FAssetRenameManager::FixReferencesAndRename(const TArray<FAssetRenameData>&
 		LoadReferencingPackages(AssetsToRename, bSoftReferencesOnly, true, ReferencingPackagesToSave, SoftReferencingObjects);
 
 		// Prompt to check out source package and all referencing packages, leave redirectors for assets referenced by packages that are not checked out and remove those packages from the save list.
-		RenamingSlowTask.EnterProgressFrame(0.1f, LOCTEXT("CheckOutPackagesSlowTask", "Check out packages..."));
+		// If revision control is disabled, this still checks if the packages are not read-only.
+		const FText& CheckOutSlowTask = ISourceControlModule::Get().IsEnabled()
+			? LOCTEXT("CheckOutPackagesSlowTask", "Check out packages...")
+			: LOCTEXT("PreparePackagesSlowTask", "Preparing packages...");
+		RenamingSlowTask.EnterProgressFrame(0.1f, CheckOutSlowTask);
 		const bool bUserAcceptedCheckout = CheckOutPackages(AssetsToRename, ReferencingPackagesToSave, bAutoCheckout);
 
 		if (bUserAcceptedCheckout || bSoftReferencesOnly)
@@ -1912,11 +1916,12 @@ void FAssetRenameManager::PerformAssetRename(TArray<FAssetRenameDataWithReferenc
 	}
 
 	const FText& UpdatingRevisionControlSlowTask = LOCTEXT("UpdatingRevisionControlSlowTask", "Updating Revision Control...");
-	RenamingSlowTask.EnterProgressFrame(0.5f, UpdatingRevisionControlSlowTask);
 
 	// Bulk update SCC status for old packages since it is faster than doing it one by one below
 	if (ISourceControlModule::Get().IsEnabled())
 	{
+		RenamingSlowTask.EnterProgressFrame(0.5f, UpdatingRevisionControlSlowTask);
+
 		TArray<FString> AllSourceFilenamesToCheck;
 
 		ISourceControlProvider& SourceControlProvider = ISourceControlModule::Get().GetProvider();
@@ -1940,24 +1945,28 @@ void FAssetRenameManager::PerformAssetRename(TArray<FAssetRenameDataWithReferenc
 			SourceControlProvider.GetState(AllSourceFilenamesToCheck, States, EStateCacheUsage::ForceUpdate);
 		}
 	}
-
-	RenamingSlowTask.EnterProgressFrame(5.0f, UpdatingRevisionControlSlowTask);
-	FScopedSlowTask BuildingRelationshipSlowTask(5.0f, LOCTEXT("BranchingRenamedAssetsSlowTask", "Building a relationship between the renamed files in Revision Control..."));
-
-	// Now branch the files in source control if possible
-	ProgressStep = 2.9f / static_cast<float>(AssetsToRename.Num());
-	for (const FAssetRenameDataWithReferencers& RenameData : AssetsToRename)
+	else
 	{
-		BuildingRelationshipSlowTask.EnterProgressFrame(ProgressStep);
+		RenamingSlowTask.EnterProgressFrame(0.5f, FText::GetEmpty());
+	}
 
-		UPackage* OldPackage = FindPackage(nullptr, *RenameData.OldObjectPath.GetLongPackageName());
-		UPackage* NewPackage = FindPackage(nullptr, *RenameData.NewObjectPath.GetLongPackageName());
+	if (ISourceControlModule::Get().IsEnabled())
+	{
+		RenamingSlowTask.EnterProgressFrame(5.0f, UpdatingRevisionControlSlowTask);
+		FScopedSlowTask BuildingRelationshipSlowTask(5.0f, LOCTEXT("BranchingRenamedAssetsSlowTask", "Building a relationship between the renamed files in Revision Control..."));
 
-		// If something went wrong when saving and the new asset does not exist on disk, don't branch it
-		// as it will just create a copy and any attempt to load it will result in crashes.
-		if (!RenameData.bOnlyFixSoftReferences && NewPackage && FPackageName::DoesPackageExist(NewPackage->GetName()))
+		// Now branch the files in source control if possible
+		ProgressStep = 2.9f / static_cast<float>(AssetsToRename.Num());
+		for (const FAssetRenameDataWithReferencers& RenameData : AssetsToRename)
 		{
-			if (ISourceControlModule::Get().IsEnabled())
+			BuildingRelationshipSlowTask.EnterProgressFrame(ProgressStep);
+
+			UPackage* OldPackage = FindPackage(nullptr, *RenameData.OldObjectPath.GetLongPackageName());
+			UPackage* NewPackage = FindPackage(nullptr, *RenameData.NewObjectPath.GetLongPackageName());
+
+			// If something went wrong when saving and the new asset does not exist on disk, don't branch it
+			// as it will just create a copy and any attempt to load it will result in crashes.
+			if (!RenameData.bOnlyFixSoftReferences && NewPackage && FPackageName::DoesPackageExist(NewPackage->GetName()))
 			{
 				ISourceControlProvider& SourceControlProvider = ISourceControlModule::Get().GetProvider();
 				const FString SourceFilename = USourceControlHelpers::PackageFilename(OldPackage);
@@ -1973,6 +1982,10 @@ void FAssetRenameManager::PerformAssetRename(TArray<FAssetRenameDataWithReferenc
 			}
 		}
 	}
+	else
+	{
+		RenamingSlowTask.EnterProgressFrame(5.0f, FText::GetEmpty());
+	}
 
 	// Clean up all packages that were left empty
 	if (PotentialPackagesToDelete.Num() > 0)
