[AssetHeaderPatcher]

- Adds support for reading and patching the AssetRegistryDependency section of the header
- Corrects the import table patching to better handle long-object import paths that may result in new entries in the ImportTable.
- Corrects FName table mapping to prevent remapping existing names in the table to a new index when the patcher adds an existing name.
- Adds redirect for PersistentLevel subobjects that need patching
- Re-enables patching FNames found in the Export Data section of the FName table that overlap header FNames. Not doing so results in common types referring to object paths not being patched leading to loading errors (e.g. FPropertyTag)
- Co-authored with matt.peters

#rnx

diff --git a/Engine/Source/Developer/AssetTools/Private/AssetHeaderPatcher.cpp b/Engine/Source/Developer/AssetTools/Private/AssetHeaderPatcher.cpp
--- a/Engine/Source/Developer/AssetTools/Private/AssetHeaderPatcher.cpp
+++ b/Engine/Source/Developer/AssetTools/Private/AssetHeaderPatcher.cpp
@@ -3,6 +3,7 @@
 #include "AssetHeaderPatcher.h"
 
 #include "Algo/Copy.h"
+#include "Algo/Sort.h"
 #include "AssetRegistry/IAssetRegistry.h"
 #include "AssetRegistry/AssetData.h"
 #include "AssetRegistry/PackageReader.h"
@@ -458,10 +459,35 @@ void FAssetHeaderPatcher::FContext::GenerateAdditionalRemappings()
 			DstNameBuilder.Append(DstPackageName);
 			DstNameBuilder.AppendChar(TEXT('.'));
 
-			FCoreRedirect PackageObjectRedirect(ECoreRedirectFlags::Option_MatchPrefix | ECoreRedirectFlags::Type_Object,
+			FCoreRedirect ObjectRedirect(ECoreRedirectFlags::Option_MatchPrefix | ECoreRedirectFlags::Type_Object,
 				FCoreRedirectObjectName(SrcNameBuilder.ToString()),
 				FCoreRedirectObjectName(DstNameBuilder.ToString()));
-			Redirects.Emplace(MoveTemp(PackageObjectRedirect));
+			Redirects.Emplace(MoveTemp(ObjectRedirect));
+		}
+
+		// Path.Object.PersistentLevel.* mapping
+		{
+			FName PersistentLevelFName(NAME_PersistentLevel);
+			SrcNameBuilder.Reset();
+			SrcNameBuilder.Append(SrcNameString);
+			SrcNameBuilder.AppendChar(TEXT('.'));
+			SrcNameBuilder.Append(SrcPackageName);
+			SrcNameBuilder.AppendChar(TEXT('.'));
+			PersistentLevelFName.AppendString(SrcNameBuilder);
+			SrcNameBuilder.AppendChar(TEXT('.'));
+
+			DstNameBuilder.Reset();
+			DstNameBuilder.Append(DstNameString);
+			DstNameBuilder.AppendChar(TEXT('.'));
+			DstNameBuilder.Append(DstPackageName);
+			DstNameBuilder.AppendChar(TEXT('.'));
+			PersistentLevelFName.AppendString(DstNameBuilder);
+			DstNameBuilder.AppendChar(TEXT('.'));
+
+			FCoreRedirect ObjectRedirect(ECoreRedirectFlags::Option_MatchPrefix | ECoreRedirectFlags::Type_Object,
+				FCoreRedirectObjectName(SrcNameBuilder.ToString()),
+				FCoreRedirectObjectName(DstNameBuilder.ToString()));
+			Redirects.Emplace(MoveTemp(ObjectRedirect));
 		}
 
 		// MaterialFunctionInterface "EditorOnlyData"
@@ -610,7 +636,7 @@ public:
 
 		if (MaybeIndex == nullptr)
 		{
-			ErrorMessage += FString::Printf(TEXT("Cannot serialize FName %s because it is not in the name table for %s\n"), *Name.ToString(), *GetArchiveName());
+			ErrorMessage += FString::Printf(TEXT("Cannot serialize FName '%s' because it is not in the name table for %s\n"), *Name.ToString(), *GetArchiveName());
 			SetCriticalError();
 			return *this;
 		}
@@ -646,7 +672,8 @@ enum class EPatchedSection
 	ExportTable,
 	SoftPackageReferencesTable,
 	ThumbnailTable,
-	AssetRegistryData
+	AssetRegistryData,
+	AssetRegistryDependencyData,
 };
 
 struct FSectionData
@@ -868,6 +895,7 @@ public:
 	bool DoPatch(FGatherableTextData& InOutGatherablerTextData);
 	bool DoPatch(FThumbnailEntry& InOutThumbnailEntry);
 	bool RemapFName(FName SrcName, FName DstName);
+	bool AddFName(FName DstName);
 	bool ShouldReplaceMountPoint(const FStringView InPath, FStringView& OutSrcMountPoint, FStringView& OutDstMountPoint);
 	FCoreRedirectObjectName GetFullObjectNameFromObjectResource(const FObjectResource& InResource, bool bIsExport, bool bWalkImportsOnly = false);
 
@@ -887,16 +915,17 @@ public:
 #if WITH_EDITORONLY_DATA
 		FName PackageName;
 #endif
+		bool bUsedInGame = true;
 	};
 
 	void GetExportTablePatches(TArray<FExportPatch>& OutExportPatches);
-	void GetImportTablePatches(TArray<FImportPatch>& OutImportPatches, int32& OutNewImportCount);
+	FAssetHeaderPatcher::EResult GetImportTablePatches(TArray<FImportPatch>& OutImportPatches, int32& OutNewImportCount);
 	void PatchExportAndImportTables(const TArray<FExportPatch>& InExportPatches, const TArray<FImportPatch>& InImportPatches, const int32 InNewImportCount);
 	void PatchNameTable();
 
 	FAssetHeaderPatcher::EResult PatchHeader();
 	FAssetHeaderPatcher::EResult PatchHeader_Deserialize();
-	void PatchHeader_PatchSections();
+	FAssetHeaderPatcher::EResult PatchHeader_PatchSections();
 	FAssetHeaderPatcher::EResult PatchHeader_WriteDestinationFile();
 	void DumpState(FStringView InDir);
 
@@ -930,6 +959,7 @@ public:
 	FPackageFileSummary Summary;
 	FName OriginalPackagePath;					 // e.g. "/MountName/TopLevelPackageName"
 	FName OriginalNonOneFilePerActorPackagePath; // e.g. "/MountName/MountName"
+	FName DstPackagePath;						 // OriginalPackagePath, or the remapped name of it if it was remapped
 	FString OriginalPrimaryAssetName;			 // e.g. "MountName"
 	bool bPatchPrimaryAssetTag;
 	bool bIsPackagePathInNametable;
@@ -965,6 +995,10 @@ public:
 		int64 SectionSize = -1;
 		UE::AssetRegistry::FDeserializePackageData PkgData;
 		TArray<FAssetRegistryObjectData> ObjectData;
+		int64 DependencyDataSectionSize = -1;
+		TMap<int32, bool> ImportIndexUsedInGame;
+		TMap<FName, bool> SoftPackageReferenceUsedInGame;
+		TArray<TPair<FName, UE::AssetRegistry::EExtraDependencyFlags>> ExtraPackageDependencies;
 	};
 	FAssetRegistryData AssetRegistryData;
 };
@@ -1131,7 +1165,11 @@ FAssetHeaderPatcher::EResult FAssetHeaderPatcherInner::PatchHeader()
 	
 	if (DumpOutputDirectory.IsEmpty())
 	{
-		PatchHeader_PatchSections();
+		Result = PatchHeader_PatchSections();
+		if (Result != EResult::Success)
+		{
+			return Result;
+		}
 	}
 	else
 	{
@@ -1142,7 +1180,11 @@ FAssetHeaderPatcher::EResult FAssetHeaderPatcherInner::PatchHeader()
 		FPaths::RemoveDuplicateSlashes(BeforeDir);
 		DumpState(BeforeDir);
 
-		PatchHeader_PatchSections();
+		Result = PatchHeader_PatchSections();
+		if (Result != EResult::Success)
+		{
+			return Result;
+		}
 
 		FString AfterDir = BaseDir / FString(TEXT("After"));
 		FPaths::RemoveDuplicateSlashes(AfterDir);
@@ -1473,6 +1515,56 @@ FAssetHeaderPatcher::EResult FAssetHeaderPatcherInner::PatchHeader_Deserialize()
 		}
 
 		AssetRegistryData.SectionSize = MemAr.Tell() - Summary.AssetRegistryDataOffset;
+
+		UE::AssetRegistry::FReadPackageDataDependenciesArgs DependenciesArgs;
+		DependenciesArgs.BinaryNameAwareArchive = &MemAr;
+		DependenciesArgs.AssetRegistryDependencyDataOffset = AssetRegistryData.PkgData.DependencyDataOffset;
+		DependenciesArgs.NumImports = Summary.ImportCount;
+		DependenciesArgs.NumSoftPackageReferences = Summary.SoftPackageReferencesCount;
+		DependenciesArgs.PackageVersion = Summary.GetFileVersionUE();
+
+		if (!UE::AssetRegistry::ReadPackageDataDependencies(DependenciesArgs))
+		{
+			UE_LOG(LogAssetHeaderPatcher, Error, TEXT("Failed to deserialize asset registry data for %s"), *SrcAsset);
+			return EResult::ErrorFailedToDeserializeSourceAsset;
+		}
+
+		if (DependenciesArgs.ImportUsedInGame.Num() != Summary.ImportCount ||
+			DependenciesArgs.SoftPackageUsedInGame.Num() != Summary.SoftPackageReferencesCount)
+		{
+			UE_LOG(LogAssetHeaderPatcher, Error,
+				TEXT("Failed to deserialize asset registry data for %s. ReadPackageDataDependencies internal error: (%d != %d || %d != %d)."),
+				*SrcAsset, DependenciesArgs.ImportUsedInGame.Num(), Summary.ImportCount,
+				DependenciesArgs.SoftPackageUsedInGame.Num(), Summary.SoftPackageReferencesCount);
+			return EResult::ErrorFailedToDeserializeSourceAsset;
+		}
+
+		if ((AssetRegistryData.PkgData.DependencyDataOffset != INDEX_NONE) != (DependenciesArgs.AssetRegistryDependencyDataSize != 0))
+		{
+			// When writing the file, we use validity of Size to decide whether we write or skip the
+			// AssetRegistryDependencyData section, but in the earlier AssetRegistryData section we write the DependencyOffset
+			// based on whether the offset is valid. To prevent corruption, we need the validity of each to match. 
+			UE_LOG(LogAssetHeaderPatcher, Error,
+				TEXT("Failed to deserialize asset registry data for %s. DependencyDataOffset (%" INT64_FMT ") != -1 does not match AssetRegistryDependencyDataSize (%" INT64_FMT ") != 0."),
+				*DstAsset, AssetRegistryData.PkgData.DependencyDataOffset, DependenciesArgs.AssetRegistryDependencyDataSize);
+			return EResult::ErrorFailedToOpenDestinationFile;
+		}
+
+		AssetRegistryData.DependencyDataSectionSize = DependenciesArgs.AssetRegistryDependencyDataSize;
+		AssetRegistryData.ImportIndexUsedInGame.Reserve(Summary.ImportCount);
+		for (int32 ImportIndex = 0; ImportIndex < Summary.ImportCount; ++ImportIndex)
+		{
+			AssetRegistryData.ImportIndexUsedInGame.Add(ImportIndex, DependenciesArgs.ImportUsedInGame[ImportIndex]);
+		}
+		check(SoftPackageReferencesTable.Num() == Summary.SoftPackageReferencesCount); // Constructed above
+		AssetRegistryData.SoftPackageReferenceUsedInGame.Reserve(Summary.SoftPackageReferencesCount);
+		for (int32 SoftPackageIndex = 0; SoftPackageIndex < Summary.SoftPackageReferencesCount; ++SoftPackageIndex)
+		{
+			AssetRegistryData.SoftPackageReferenceUsedInGame.Add(SoftPackageReferencesTable[SoftPackageIndex],
+				DependenciesArgs.SoftPackageUsedInGame[SoftPackageIndex]);
+		}
+
+		AssetRegistryData.ExtraPackageDependencies = MoveTemp(DependenciesArgs.ExtraPackageDependencies);
 	}
 
 	return EResult::Success;
@@ -1682,10 +1774,35 @@ bool FAssetHeaderPatcherInner::DoPatch(FString& InOutString)
 	return bDidPatch;
 }
 
+bool FAssetHeaderPatcherInner::AddFName(FName DstName)
+{
+	if (DstName == NAME_None)
+	{
+		return false;
+	}
+
+	FNameEntryId DstComparisonId = DstName.GetDisplayIndex();
+	FNameEntryId* RemappedFName = RenameMap.Find(DstComparisonId);
+	if (RemappedFName)
+	{
+		// We hit a case where we thought we needed to change the name in the NameTable
+		// but have now discovered some part of the header needs to use the old name. In such a case
+		// remove the remap and make it an add instead.
+		AddedNames.Add(*RemappedFName);
+		RenameMap.Remove(DstComparisonId);
+	}
+	else
+	{
+		AddedNames.Add(DstComparisonId);
+	}
+	UnchangedNames.Add(DstComparisonId);
+	return true;
+}
+
 bool FAssetHeaderPatcherInner::RemapFName(FName SrcName, FName DstName)
 {
 	// None won't appear in the NameTable so skip any None passed in here
-	if(SrcName == NAME_None)
+	if (SrcName == NAME_None)
 	{
 		return false;
 	}
@@ -1697,47 +1814,46 @@ bool FAssetHeaderPatcherInner::RemapFName(FName SrcName, FName DstName)
 	FNameEntryId DstComparisonId = DstName.GetDisplayIndex();
 
 	// Since we do fuzzy matching against AssetRegistryTag data, we might get SrcNames not in the FName table as input; ignore these.
-	if(!NameToIndexMap.Contains(SrcComparisonId))
+	if (!NameToIndexMap.Contains(SrcComparisonId))
 	{
 		return false;
 	}
 
-	bool bIsExportDataFName = NameToIndexMap[SrcComparisonId] < Summary.NamesReferencedFromExportDataCount;
-	bool bNeedRemap = SrcComparisonId != DstComparisonId;
-	FNameEntryId* RemappedFName = RenameMap.Find(SrcComparisonId);
+	// Previously it was thought that any FName at the front of the NameTable is an FName used by export data and should not be patched
+	// since we can't know the context in which the user uses the FName. This is true, however it's quite common for PropertyTags and other
+	// common types to be "ExportData" that refers to paths that will be patched and must stay in sync. As such if we come across one of these
+	// paths when patching exports normally, we must replace it. We still do not do a generic walk of the FNameTable for patching without context
+	// as that will be error-prone however we allow patching the "ExportData" section of the NameTable when we know there is an overlap with
+	// header FNames and ExportData FNames.
+	constexpr bool bIsExportDataFName = false; // NameToIndexMap[SrcComparisonId] < Summary.NamesReferencedFromExportDataCount;
 
+	bool bNeedRemap = SrcComparisonId != DstComparisonId;
 	if (!bNeedRemap)
 	{
-		UnchangedNames.Add(SrcComparisonId);
-
-		if (RemappedFName)
-		{
-			// We hit a case where we thought we needed to change the name in the NameTable
-			// but have now discovered some part of the header needs to use the old name. In such a case
-			// remove the remap and make it an add instead.
-			AddedNames.Add(*RemappedFName);
-			RenameMap.Remove(SrcComparisonId);
-		}
+		AddFName(SrcName);
+		return false;
 	}
 	else
 	{
-		// We already have a _different_ remapping, a header entry is still relying on the original name or this is a name used by the Export payload in which case we can't safely modify it. That is fine; we might have used the same FName in more than one place.
-		// We need to be certain we are renaming only the names we care about in context. If not, shared names in the NameTable might change
-		// incorrectly (e.g. A class FName may have matched a Package FName, thus sharing a NameTable entry, but after patching it's possible _only_ 
-		// the Package name has changed. In such a case we don't want to rename the class name inadvertently by patching the shared NameTable entry. 
-		// If we have a mismatch with the new patched name, record the new name and we will append it to the NameTable later.
+		FNameEntryId* RemappedFName = RenameMap.Find(SrcComparisonId);
+
 		bool bForceAddName = bIsExportDataFName || (RemappedFName && *RemappedFName != DstComparisonId) || UnchangedNames.Contains(SrcComparisonId);
 		if (bForceAddName)
 		{
+			// We already have a _different_ remapping, a header entry is still relying on the original name or this is a name used by the Export payload
+			// in which case we can't safely modify it. That is fine; we might have used the same FName in more than one place.
+			// We need to be certain we are renaming only the names we care about in context. If not, shared names in the NameTable might change
+			// incorrectly (e.g. A class FName may have matched a Package FName, thus sharing a NameTable entry, but after patching it's possible _only_ 
+			// the Package name has changed. In such a case we don't want to rename the class name inadvertently by patching the shared NameTable entry. 
+			// If we have a mismatch with the new patched name, record the new name and we will append it to the NameTable later.
 			AddedNames.Add(DstComparisonId);
 		}
 		else
 		{
 			RenameMap.Add(SrcComparisonId, DstComparisonId);
 		}
+		return true;
 	}
-
-	return bNeedRemap;
 }
 
 bool FAssetHeaderPatcherInner::DoPatch(FName& InOutName)
@@ -1772,31 +1888,33 @@ void FAssetHeaderPatcherInner::GetExportTablePatches(TArray<FExportPatch>& OutEx
 	}
 }
 
-void FAssetHeaderPatcherInner::GetImportTablePatches(TArray<FImportPatch>& OutImportPatches, int32& OutNewImportCount)
+static FName GetObjectResourceNameFromCoreRedirectObjectName(const FCoreRedirectObjectName& InName)
+{
+	return !InName.ObjectName.IsNone() ? InName.ObjectName : InName.PackageName;
+}
+
+FAssetHeaderPatcher::EResult FAssetHeaderPatcherInner::GetImportTablePatches(TArray<FImportPatch>& OutImportPatches, int32& OutNewImportCount)
 {
 	OutNewImportCount = 0;
 	OutImportPatches.Reserve(ImportTable.Num());
 	
-	TArray<int32> PatchedImports;
-	PatchedImports.Reserve(ImportTable.Num());
-
-	for (int32 i = 0; i < ImportTable.Num(); ++i)
+	struct FPatchDataForImport
 	{
-		FObjectImport& Import = ImportTable[i];
-
-		FCoreRedirectObjectName SrcResourceName = GetFullObjectNameFromObjectResource(Import, false /*bIsExport*/, true /*bWalkImportsOnly*/);
-		FCoreRedirectObjectName DstResourceName = FCoreRedirects::GetRedirectedName(ECoreRedirectFlags::Type_AllMask, SrcResourceName);
-		ImportNameToImportTableIndexLookup.Add(SrcResourceName, i);
-		ImportTablePatchedNames.Add({ SrcResourceName, DstResourceName });
+		int32 PatchIndex = INDEX_NONE;
+		FCoreRedirectObjectName SrcImportPath;
+		FCoreRedirectObjectName DstImportPath;
+		bool bPatched = false;
+		bool bRedirectEvaluated = false;
+	};
+	TArray<FPatchDataForImport> ImportIndexToPatchData;
+	ImportIndexToPatchData.SetNum(ImportTable.Num());
 
-		if (SrcResourceName != DstResourceName)
-		{
-			PatchedImports.Add(i);
-		}
+	for (int32 ImportIndex = 0; ImportIndex < ImportTable.Num(); ++ImportIndex)
+	{
+		FObjectImport& Import = ImportTable[ImportIndex];
+		FPatchDataForImport& PatchDataForImport = ImportIndexToPatchData[ImportIndex];
 
-		// While we are here we patch FNames that have no impact on the FObjectResource walking we will do below.
-		// This saves us from having to output FImportPatch entries for entries that don't have their ObjectPath patched
-		// but may have changed the package name or class path.
+		// For each FObjectResource, immediately patch the FNames that do not impact other imports or exports.
 		const FCoreRedirectObjectName SrcImportClass(Import.ClassName, NAME_None, Import.ClassPackage);
 		const FCoreRedirectObjectName DstImportClass = FCoreRedirects::GetRedirectedName(ECoreRedirectFlags::Type_Class | ECoreRedirectFlags::Type_Package, SrcImportClass);
 
@@ -1812,121 +1930,248 @@ void FAssetHeaderPatcherInner::GetImportTablePatches(TArray<FImportPatch>& OutIm
 		Import.PackageName = RemapFName(SrcImportPackageName.PackageName, DstImportPackageName.PackageName)
 			? DstImportPackageName.PackageName : SrcImportPackageName.PackageName;
 #endif
+
+		// Look up whether there is a specific redirect for the full path of the import.
+		PatchDataForImport.SrcImportPath = GetFullObjectNameFromObjectResource(Import, false /*bIsExport*/, false /*bWalkImportsOnly*/);
+		PatchDataForImport.DstImportPath = FCoreRedirects::GetRedirectedName(ECoreRedirectFlags::Type_AllMask, PatchDataForImport.SrcImportPath);
+		// If there is not a specific redirect, the object might still be redirected because its outer is redirected. All objects within
+		// a redirected outer are moved along with that outer, but the CoreRedirects system does not find the redirect of the outer unless
+		// we look for it specifically. We handle this case in the loop below, after we have processed all imports looking for redirects
+		// specific to them.
+
+		// Record the remapping of the FName in the FObjectResource (whether or not it was redirected)
+		FName SrcObjectResourceName = GetObjectResourceNameFromCoreRedirectObjectName(PatchDataForImport.SrcImportPath);
+		FName DstObjectResourceName = GetObjectResourceNameFromCoreRedirectObjectName(PatchDataForImport.DstImportPath);
+		RemapFName(SrcObjectResourceName, DstObjectResourceName);
+
+		// If there was a redirect specifically for this import, create an FImportPatch
+		PatchDataForImport.bPatched = PatchDataForImport.SrcImportPath != PatchDataForImport.DstImportPath;
+		if (PatchDataForImport.bPatched)
+		{
+			PatchDataForImport.bRedirectEvaluated = true;
+			PatchDataForImport.PatchIndex = OutImportPatches.Num();
+
+			FImportPatch& ImportPatch = OutImportPatches.Emplace_GetRef();
+			ImportPatch.ClassName = Import.ClassName;
+			ImportPatch.ClassPackage = Import.ClassPackage;
+#if WITH_EDITORONLY_DATA
+			ImportPatch.PackageName = Import.PackageName;
+#endif
+			bool* UsedInGame = AssetRegistryData.ImportIndexUsedInGame.Find(ImportIndex);
+			ImportPatch.bUsedInGame = UsedInGame ? *UsedInGame : true;
+			ImportPatch.TableIndex = ImportIndex;
+			ImportPatch.ObjectName = DstObjectResourceName;
+			ImportPatch.OuterIndex = Import.OuterIndex;
+		}
 	}
 
-	// Now that we know which import entries won't change we can review our patched Imports and see how they need to be patched
-	// (e.g. did the ObjectName change and/or did the Outer change) and determine if we need new even more new imports (in which case we add them)
-	for (auto It = PatchedImports.CreateIterator(); It; It++)
+	// Now that redirects specifically for each import have been detected, go through the list of imports again and assign
+	// DstImportPath based on their (possibly redirected) outer.
+	for (int32 ImportIndex = 0; ImportIndex < ImportTable.Num(); ++ImportIndex)
 	{
-		const int32 ImportIndexToPatch = *It;
-		It.RemoveCurrent();
-
-		// Go outer-most to inner-most from the patched name and see if we have all necessary imports patches.
-		// If we already have the patched Import's name in our lookup table, no patching needed; record the OuterIndex 
-		// for the next inner import outered to this one to use. If not, create a new import patch.
-		FPackageIndex OuterIndex = FPackageIndex();
-		const FCoreRedirectObjectName& SrcResourceName = ImportTablePatchedNames[ImportIndexToPatch].Key;
-		const FCoreRedirectObjectName& DstResourceName = ImportTablePatchedNames[ImportIndexToPatch].Value;
+		FObjectImport& Import = ImportTable[ImportIndex];
+		FPatchDataForImport& PatchDataForImport = ImportIndexToPatchData[ImportIndex];
 
-		auto CreatePatch = [this, ImportIndexToPatch, &OuterIndex, &SrcResourceName, 
-			&PatchedImports, &OutImportPatches, &OutNewImportCount] (const FCoreRedirectObjectName SrcImportName, const FCoreRedirectObjectName DstImportName, const FName DstImportObjectName)
+		if (!PatchDataForImport.bRedirectEvaluated)
+		{
+			// Recursively evaluate all outers of the import, and then evaluate the import
+			TArray<int32, TInlineAllocator<10>> ImportsToEvaluate;
+			ImportsToEvaluate.Add(ImportIndex);
+			for (FPackageIndex OuterIndex = Import.OuterIndex; !OuterIndex.IsNull(); )
 			{
-				check(!DstImportObjectName.IsNone());
-
-				// If we are going to patch this SrcImportName later, skip so we can ensure we don't inadvertently add unnecessary entries
-				int32* SrcIndex = ImportNameToImportTableIndexLookup.Find(SrcImportName);
-				if (SrcIndex && PatchedImports.Contains(*SrcIndex))
+				if (OuterIndex.IsExport())
 				{
-					// Update our outer so any patched imports for objects outered to this one are linked correctly
-					OuterIndex = FPackageIndex::FromImport(*SrcIndex);
-					return;
+					// We don't have data on the export, so stop.
+					// TODO: This is incorrect, we need to get the remapped destination from the export to know
+					// the correct destination for the import.
+					break;
 				}
+				check(OuterIndex.IsImport());
+				int32 OuterImportIndex = OuterIndex.ToImport();
+				FObjectImport& OuterImport = ImportTable[OuterImportIndex];
+				FPatchDataForImport& PatchDataForOuter = ImportIndexToPatchData[OuterImportIndex];
+				if (PatchDataForOuter.bRedirectEvaluated)
+				{
+					break;
+				}
+				ImportsToEvaluate.Add(OuterImportIndex);
+				OuterIndex = OuterImport.OuterIndex;
+			}
 
-				FImportPatch ImportPatch;
-
-				// Prepopulate the patch with the existing import data. We may have patched fields like
-				// ClassName already and want those to be copied over when patching, or used as the 
-				// new value when adding new entries
-				const FObjectImport& SrcImport = ImportTable[ImportIndexToPatch];
-				ImportPatch.ClassName = SrcImport.ClassName;
-				ImportPatch.ClassPackage = SrcImport.ClassPackage;
-#if WITH_EDITORONLY_DATA
-				ImportPatch.PackageName = SrcImport.PackageName;
-#endif
+			while (!ImportsToEvaluate.IsEmpty())
+			{
+				int32 ImportIndexToEvaluate = ImportsToEvaluate.Pop(EAllowShrinking::No);
+				FObjectImport& ImportToEvaluate = ImportTable[ImportIndexToEvaluate];
+				FPatchDataForImport& PatchDataToEvaluate = ImportIndexToPatchData[ImportIndexToEvaluate];
+				if (PatchDataToEvaluate.bRedirectEvaluated)
+				{
+					// Shouldn't be possible, but somehow we already evaluated this, cycle in the outer chain?
+					continue;
+				}
 
-				// If import entry matches exactly the one we are patching (not just a portion of
-				// it's path) then we can overwrite the existing table entry instead of adding a new one
-				int32 PatchImportTableIndex;
-				if (SrcImportName == SrcResourceName)
+				if (ImportToEvaluate.OuterIndex.IsNull())
+				{
+					// This import has no outer, it is a package. We did not find a redirect specifically for it,
+					// and it has no outer to move along with, so it has no redirect. Mark it evaluated and do nothing.
+					PatchDataToEvaluate.bRedirectEvaluated = true;
+				}
+				else if (ImportToEvaluate.OuterIndex.IsExport())
 				{
-					PatchImportTableIndex = ImportIndexToPatch;
+					// We don't have data on the export, so we cannot read the transformed path for the import under it.
+					// TODO: This is incorrect, we need to get the remapped destination from the export to know
+					// the correct destination for the import.
+					PatchDataToEvaluate.bRedirectEvaluated = true;
 				}
 				else
 				{
-					PatchImportTableIndex = ImportTable.Num() + OutNewImportCount++;
+					int32 OuterImportIndex = ImportToEvaluate.OuterIndex.ToImport();
+					FObjectImport& OuterImport = ImportTable[OuterImportIndex];
+					FPatchDataForImport& PatchDataForOuter = ImportIndexToPatchData[OuterImportIndex];
+					PatchDataToEvaluate.DstImportPath = FCoreRedirectObjectName::AppendObjectName(
+						PatchDataForOuter.DstImportPath, ImportToEvaluate.ObjectName);
+					PatchDataToEvaluate.bRedirectEvaluated = true;
 				}
-				
-				ImportPatch.TableIndex = PatchImportTableIndex;
-				ImportPatch.ObjectName = DstImportObjectName;
-				RemapFName(SrcImport.ObjectName, DstImportObjectName);
-				ImportPatch.OuterIndex = OuterIndex;
-
-				// Update our outer so any patched imports for objects outered to this one are linked correctly
-				OuterIndex = FPackageIndex::FromImport(PatchImportTableIndex);
-
-				check(!ImportNameToImportTableIndexLookup.Contains(DstImportName));
-				ImportNameToImportTableIndexLookup.Add(DstImportName, PatchImportTableIndex);
-				OutImportPatches.Add(MoveTemp(ImportPatch));
-			};
+			}
+		}
+
+		check(PatchDataForImport.bRedirectEvaluated);
+
+		// Record that the (possibly redirected) full path of the destination of this import is in this import index
+		ImportNameToImportTableIndexLookup.Add(PatchDataForImport.DstImportPath, ImportIndex);
+		ImportTablePatchedNames.Add({ PatchDataForImport.SrcImportPath, PatchDataForImport.DstImportPath });
+	}
+
+	// Loop over all imports, and update their outerindex. If an import for their outer does not already exist, add it on to the end.
+	// Keep iterating over the added elements that we add on to the end until we don't add any more.
+	for (int32 ImportIndex = 0;
+		ImportIndex < ImportIndexToPatchData.Num(); // Num can change during the loop
+		++ImportIndex)
+	{
+		FCoreRedirectObjectName ImportPath = ImportIndexToPatchData[ImportIndex].DstImportPath;
+		// Do not create a pointer into ImportIndexToPatchData after we have done any necessary adds into ImportIndexToPatchData
+		FPatchDataForImport* PatchDataForImport = nullptr;
 
-		const FCoreRedirectObjectName DstPackageImportName(NAME_None, NAME_None, DstResourceName.PackageName);
-		if (int32* ImportIndex = ImportNameToImportTableIndexLookup.Find(DstPackageImportName))
+		FPackageIndex DstOuterIndex;
+		if (ImportPath.ObjectName.IsNone())
 		{
-			OuterIndex = FPackageIndex::FromImport(*ImportIndex);
+			// A package, no outer
+			DstOuterIndex = FPackageIndex();
 		}
 		else
 		{
-			const FCoreRedirectObjectName SrcPackageImportName(NAME_None, NAME_None, SrcResourceName.PackageName);
-			CreatePatch(SrcPackageImportName, DstPackageImportName, DstPackageImportName.PackageName);
-		}
-		
-		if (!DstResourceName.OuterName.IsNone())
-		{
-			TStringBuilder<NAME_SIZE> ObjectString;
-			DstResourceName.PackageName.ToString(ObjectString);
-			ObjectString.AppendChar(TEXT('.'));
-			const int32 PackagePrefixLen = ObjectString.Len();
-
-			TStringBuilder<NAME_SIZE> OuterString;
-			DstResourceName.OuterName.ToString(OuterString);
-			FStringView OuterStringView = OuterString.ToView();
-			int32 DelimiterPos = OuterStringView.Len();
-
-			do
+			FCoreRedirectObjectName OuterPath = FCoreRedirectObjectName::GetParent(ImportPath);
+			if (OuterPath.PackageName == this->DstPackagePath)
 			{
-				ObjectString.RemoveSuffix(ObjectString.Len() - PackagePrefixLen);
-				OuterStringView.LeftInline(DelimiterPos);
-				ObjectString.Append(OuterStringView);
-				FCoreRedirectObjectName DstOuterImportName(ObjectString.ToString());
-				if (int32* ImportIndex = ImportNameToImportTableIndexLookup.Find(DstOuterImportName))
+				// The outer of the import is an export (this happens when the import is an external actor package,
+				// and the external actor is a child of a ULevel in this map package).
+				// Keep the old OuterIndex. TODO: Read the correct package path from a map that we construct from remapped export name to export index,
+				// so that we can handle replacing the outerindex for new outer imports
+				PatchDataForImport = &ImportIndexToPatchData[ImportIndex];
+				if (PatchDataForImport->bPatched)
 				{
-					OuterIndex = FPackageIndex::FromImport(*ImportIndex);
+					DstOuterIndex = OutImportPatches[PatchDataForImport->PatchIndex].OuterIndex;
 				}
 				else
 				{
-					const FCoreRedirectObjectName SrcOuterImportName(SrcResourceName.OuterName, NAME_None, SrcResourceName.PackageName);
-					CreatePatch(SrcOuterImportName, DstOuterImportName, DstOuterImportName.ObjectName);
+					// The only unpatched imports have to be ones from the ImportTable of the pretransformed package
+					check(ImportIndex < ImportTable.Num());
+					DstOuterIndex = ImportTable[ImportIndex].OuterIndex;
 				}
-			} while (OuterStringView.FindLastChar(TEXT('.'), DelimiterPos));
+			}
+			else
+			{
+				// The outer of an import is an import, and we need to find or add it in our importtable
+				int32& ExistingOuterIndex = ImportNameToImportTableIndexLookup.FindOrAdd(OuterPath, INDEX_NONE);
+				if (ExistingOuterIndex == INDEX_NONE)
+				{
+					// If the outer is not already in the list of imports, add on another FImportPatch to hold the outer
+					int32 OuterPatchIndex = OutImportPatches.Num();
+					FImportPatch& OuterPatch = OutImportPatches.Emplace_GetRef();
+					int32 OuterImportIndex = ImportIndexToPatchData.Num();
+					FPatchDataForImport& OuterPatchData = ImportIndexToPatchData.Emplace_GetRef();
+					OuterPatchData.bPatched = true;
+					OuterPatchData.DstImportPath = OuterPath;
+					OuterPatchData.PatchIndex = OuterPatchIndex;
+					ExistingOuterIndex = OuterImportIndex;
+
+					if (OuterPath.ObjectName.IsNone())
+					{
+						// Outer is a package
+						OuterPatch.ClassName = FName(TEXT("Package"));
+						OuterPatch.ClassPackage = FName(TEXT("/Script/CoreUObject"));
+					}
+					else
+					{
+						// Outer is an object inside a package.
+						// We don't know the class name and package of the outer, so set it to /Script/CoreUObject.Object.
+						// TODO: Is there anyway to find this out? A better guess is to copy it from the previous outer.
+						OuterPatch.ClassName = FName(TEXT("Object"));
+						OuterPatch.ClassPackage = FName(TEXT("/Script/CoreUObject"));
+					}
+					AddFName(OuterPatch.ClassName);
+					AddFName(OuterPatch.ClassPackage);
+#if WITH_EDITORONLY_DATA
+					// We don't have any information about whether the outer object has an external package. For now we
+					// just don't support that case for redirects, and assume the outer is not in an external package.
+					OuterPatch.PackageName = NAME_None;
+#endif
+					// Set the Outer to UsedInGame if the import that has it as the outer is UsedInGame.
+					PatchDataForImport = &ImportIndexToPatchData[ImportIndex];
+					if (PatchDataForImport->bPatched)
+					{
+						OuterPatch.bUsedInGame = OutImportPatches[PatchDataForImport->PatchIndex].bUsedInGame;
+					}
+					else
+					{
+						bool* UsedInGame = AssetRegistryData.ImportIndexUsedInGame.Find(ImportIndex);
+						OuterPatch.bUsedInGame = UsedInGame ? *UsedInGame : true;
+					}
+					OuterPatch.TableIndex = OuterImportIndex;
+					OuterPatch.ObjectName = GetObjectResourceNameFromCoreRedirectObjectName(OuterPath);
+					AddFName(OuterPatch.ObjectName);
+					// OuterIndex of the outer is not yet known. Set to null. If the outer is a package this is correct,
+					// otherwise it is incorrect and we have to overwrite it later when we reach the outer in our iteration over
+					// ImportIndexToPatchData.
+					OuterPatch.OuterIndex = FPackageIndex();
+				}
+				DstOuterIndex = FPackageIndex::FromImport(ExistingOuterIndex);
+			}
 		}
 
-		FCoreRedirectObjectName DstObjectImportName(DstResourceName.ObjectName, DstResourceName.OuterName, DstResourceName.PackageName);
-		if (ImportNameToImportTableIndexLookup.Find(DstObjectImportName) == nullptr)
+		// If we have already patched, just assign the outer index.
+		PatchDataForImport = &ImportIndexToPatchData[ImportIndex];
+		if (PatchDataForImport->bPatched)
 		{
-			const FCoreRedirectObjectName SrcObjectImportName(SrcResourceName.ObjectName, SrcResourceName.OuterName, SrcResourceName.PackageName);
-			CreatePatch(SrcObjectImportName, DstObjectImportName, DstObjectImportName.ObjectName);
+			OutImportPatches[PatchDataForImport->PatchIndex].OuterIndex = DstOuterIndex;
+		}
+		else
+		{
+			// The only unpatched imports have to be ones from the ImportTable of the pretransformed package
+			check(ImportIndex < ImportTable.Num());
+			FObjectImport& Import = ImportTable[ImportIndex];
+
+			// If the outer has changed and we have not already patched, turn this into a patch.
+			if (Import.OuterIndex != DstOuterIndex)
+			{
+				PatchDataForImport->bPatched = true;
+				PatchDataForImport->PatchIndex = OutImportPatches.Num();
+				FImportPatch& ImportPatch = OutImportPatches.Emplace_GetRef();
+				ImportPatch.ClassName = Import.ClassName;
+				ImportPatch.ClassPackage = Import.ClassPackage;
+#if WITH_EDITORONLY_DATA
+				ImportPatch.PackageName = Import.PackageName;
+#endif
+				bool* UsedInGame = AssetRegistryData.ImportIndexUsedInGame.Find(ImportIndex);
+				ImportPatch.bUsedInGame = UsedInGame ? *UsedInGame : true;
+				ImportPatch.TableIndex = ImportIndex;
+				ImportPatch.ObjectName = Import.ObjectName;
+				ImportPatch.OuterIndex = DstOuterIndex;
+			}
 		}
 	}
+
+	OutNewImportCount = ImportIndexToPatchData.Num() - ImportTable.Num();
+	return FAssetHeaderPatcher::EResult::Success;
 }
 
 void FAssetHeaderPatcherInner::PatchExportAndImportTables(const TArray<FExportPatch>& InExportPatches, const TArray<FImportPatch>& InImportPatches, const int32 InNewImportCount)
@@ -1954,6 +2199,7 @@ void FAssetHeaderPatcherInner::PatchExportAndImportTables(const TArray<FExportPa
 		Import.OldClassName = NAME_None;
 		Import.PackageName = ImportPatch.PackageName;
 #endif
+		AssetRegistryData.ImportIndexUsedInGame.Add(Index, ImportPatch.bUsedInGame);
 	}
 
 	for (const FExportPatch& ExportPatch : InExportPatches)
@@ -2002,20 +2248,29 @@ void FAssetHeaderPatcherInner::PatchNameTable()
 		FNameEntryId SrcName = Pair.Key;
 		FNameEntryId DstName = Pair.Value;
 		int32* pSrcIndex = NameToIndexMap.Find(SrcName);
-		checkf(pSrcIndex && *pSrcIndex < NameTable.Num(), TEXT("An FName remapping was done for a name (%s) not in the NameTable."), *FName::CreateFromDisplayId(DstName, NAME_NO_NUMBER_INTERNAL).ToString());
+		checkf(pSrcIndex && *pSrcIndex < NameTable.Num(), TEXT("An FName remapping was done for a name (%s) not in the NameTable."),
+			*FName::CreateFromDisplayId(DstName, NAME_NO_NUMBER_INTERNAL).ToString());
 		int32 SrcIndex = *pSrcIndex;
 
 		NameTable[SrcIndex] = FName::CreateFromDisplayId(DstName, NAME_NO_NUMBER_INTERNAL);
+		NameToIndexMap.Remove(SrcName);
 		NameToIndexMap.Add(DstName, SrcIndex);
 	}
 
 	for (FNameEntryId NewName : AddedNames)
 	{
+		if (NameToIndexMap.Contains(NewName))
+		{
+			// Definition of an AddedName is that an original name was remapped to a name, and we couldn't remove the
+			// original name for one of several reasons, so we want to add the remapped name to the nametable.
+			// But it is possible that remapped name already exists, so we don't need to add it.
+			continue;
+		}
 		FName NewFName = FName::CreateFromDisplayId(NewName, NAME_NO_NUMBER_INTERNAL);
 		int32 NameTableIndex = NameTable.Num();
 
 		NameTable.Add(NewFName);
-		NameToIndexMap.Add(NewFName.GetDisplayIndex(), NameTableIndex);
+		NameToIndexMap.Add(NewName, NameTableIndex);
 	}
 
 	Summary.NameCount = NameTable.Num();
@@ -2228,7 +2483,7 @@ bool FAssetHeaderPatcherInner::DoPatch(FThumbnailEntry& InThumbnailEntry)
 	return bPatched;
 }
 
-void FAssetHeaderPatcherInner::PatchHeader_PatchSections()
+FAssetHeaderPatcher::EResult FAssetHeaderPatcherInner::PatchHeader_PatchSections()
 {
 	// Package Summary
 	{
@@ -2238,6 +2493,7 @@ void FAssetHeaderPatcherInner::PatchHeader_PatchSections()
 		// This is a string, so we do not want to Remap the patched name unless it's a non-OFPA
 		// package, in which case there will be an FName entry for this path
 		Summary.PackageName = DstPackageName.PackageName.ToString();
+		DstPackagePath = DstPackageName.PackageName;
 		
 		// It seems that non-OFPA packages tend to have the package name in the nametable, 
 		// however it isn't a guarantee, so we confirm the name is there before remapping and
@@ -2256,7 +2512,11 @@ void FAssetHeaderPatcherInner::PatchHeader_PatchSections()
 	int32 NewImportCount = 0;
 	TArray<FImportPatch> ImportPatches;
 	GetExportTablePatches(ExportPatches);
-	GetImportTablePatches(ImportPatches, NewImportCount);
+	FAssetHeaderPatcher::EResult Result = GetImportTablePatches(ImportPatches, NewImportCount);
+	if (Result != FAssetHeaderPatcher::EResult::Success)
+	{
+		return Result;
+	}
 	PatchExportAndImportTables(ExportPatches, ImportPatches, NewImportCount);
 
 	// Soft paths
@@ -2271,11 +2531,14 @@ void FAssetHeaderPatcherInner::PatchHeader_PatchSections()
 		DoPatch(GatherableTextData);
 	}
 
-	// Soft Package Reference's
+	// Soft Package References
 	for (FName& Reference : SoftPackageReferencesTable)
 	{
 		FCoreRedirectObjectName SrcPackagePath(NAME_None, NAME_None, Reference);
 		FCoreRedirectObjectName DstPackageName = FCoreRedirects::GetRedirectedName(ECoreRedirectFlags::Type_Package, SrcPackagePath);
+		bool* UsedInGame = AssetRegistryData.SoftPackageReferenceUsedInGame.Find(SrcPackagePath.PackageName);
+		bool bUsedInGame = UsedInGame ? *UsedInGame : true;
+		AssetRegistryData.SoftPackageReferenceUsedInGame.Add(DstPackageName.PackageName, bUsedInGame);
 		if (RemapFName(SrcPackagePath.PackageName, DstPackageName.PackageName))
 		{
 			Reference = DstPackageName.PackageName;
@@ -2298,7 +2561,7 @@ void FAssetHeaderPatcherInner::PatchHeader_PatchSections()
 		DoPatch(ThumbnailEntry);
 	}
 
-	// Asset Register Data
+	// Asset Registry Data
 	for (FAssetRegistryObjectData& ObjData : AssetRegistryData.ObjectData)
 	{
 		// ObjectPath is a toss-up. 
@@ -2426,8 +2689,61 @@ void FAssetHeaderPatcherInner::PatchHeader_PatchSections()
 		}
 	}
 
+	// AssetRegistryDependencyData
+	{
+		using namespace UE::AssetRegistry;
+
+		// Move into a TMap to avoid O(n^2) insertions/removals
+		TMap<FName, EExtraDependencyFlags> Dependencies;
+		Dependencies.Reserve(AssetRegistryData.ExtraPackageDependencies.Num());
+		for (const TPair<FName, EExtraDependencyFlags>& Pair : AssetRegistryData.ExtraPackageDependencies)
+		{
+			Dependencies.Add(Pair.Key, Pair.Value);
+		}
+
+		TArray<TPair<FName, EExtraDependencyFlags>> AddedKeys;
+		TSet<FName> RemovedKeys;
+		for (const TPair<FName, EExtraDependencyFlags>& Pair : Dependencies)
+		{
+			FCoreRedirectObjectName SrcPackagePath(NAME_None, NAME_None, Pair.Key);
+			FCoreRedirectObjectName DstPackageName =
+				FCoreRedirects::GetRedirectedName(ECoreRedirectFlags::Type_Package, SrcPackagePath);
+			if (RemapFName(SrcPackagePath.PackageName, DstPackageName.PackageName))
+			{
+				AddedKeys.Add({ DstPackageName.PackageName, Pair.Value});
+				RemovedKeys.Add(SrcPackagePath.PackageName);
+			}
+		}
+		if (!AddedKeys.IsEmpty() || !RemovedKeys.IsEmpty())
+		{
+			for (TPair<FName, EExtraDependencyFlags>& Pair : AddedKeys)
+			{
+				// If the added name already existed, take the union of the old and new values
+				EExtraDependencyFlags& Existing = Dependencies.FindOrAdd(Pair.Key, EExtraDependencyFlags::None);
+				Existing |= Pair.Value;
+
+				// If an added key adds back a key that was removed, then no longer remove the key
+				RemovedKeys.Remove(Pair.Key);
+			}
+			for (FName RemoveKey : RemovedKeys)
+			{
+				Dependencies.Remove(RemoveKey);
+			}
+
+			// Return the remapped values to the array and restore sortedness
+			AssetRegistryData.ExtraPackageDependencies = Dependencies.Array();
+			Algo::Sort(AssetRegistryData.ExtraPackageDependencies,
+				[](const TPair<FName, EExtraDependencyFlags>& A, const TPair<FName, EExtraDependencyFlags>& B)
+				{
+					return A.Key.LexicalLess(B.Key);
+				});
+		}
+	}
+
 	// Do nametable patching last since we want to ensure we have determined all the remappings necessary
 	PatchNameTable();
+
+	return FAssetHeaderPatcher::EResult::Success;
 }
 
 FAssetHeaderPatcher::EResult FAssetHeaderPatcherInner::PatchHeader_WriteDestinationFile()
@@ -2446,6 +2762,7 @@ FAssetHeaderPatcher::EResult FAssetHeaderPatcherInner::PatchHeader_WriteDestinat
 		{ EPatchedSection::SearchableNamesMap,			Summary.SearchableNamesOffset,			HeaderInformation.SearchableNamesMapSize,			false	},
 		{ EPatchedSection::ThumbnailTable,				Summary.ThumbnailTableOffset,			HeaderInformation.ThumbnailTableSize,				false	},
 		{ EPatchedSection::AssetRegistryData,			Summary.AssetRegistryDataOffset,		AssetRegistryData.SectionSize,						true	},
+		{ EPatchedSection::AssetRegistryDependencyData,	AssetRegistryData.PkgData.DependencyDataOffset,	AssetRegistryData.DependencyDataSectionSize,		false	},
 	};
 
 	const int32 SourceTotalHeaderSize = Summary.TotalHeaderSize;
@@ -2459,18 +2776,20 @@ FAssetHeaderPatcher::EResult FAssetHeaderPatcherInner::PatchHeader_WriteDestinat
 		// Verify sections are ordered as expected
 		if (SourceSection.Offset < 0 || (SourceSection.bRequired && (SourceSection.Offset < PrevSection.Offset)))
 		{
-			UE_LOG(LogAssetHeaderPatcher, Error, TEXT("Unexpected section order for %s (%d %zd < %zd) "), *SrcAsset, SectionIdx, SourceSection.Offset, PrevSection.Offset);
+			UE_LOG(LogAssetHeaderPatcher, Error, TEXT("Unexpected section order for %s (%d %" INT64_FMT " < %" INT64_FMT ") "),
+				*SrcAsset, SectionIdx, SourceSection.Offset, PrevSection.Offset);
 			return EResult::ErrorUnexpectedSectionOrder;
 		}
 	}
 
 	// Ensure the required sections have data
-	for (const FSectionData& SourceSection : SourceSections)
+	for (int32 SectionIdx = 0; SectionIdx < UE_ARRAY_COUNT(SourceSections); ++SectionIdx)
 	{
-		// skip processing empty non required chunks.
+		const FSectionData& SourceSection = SourceSections[SectionIdx];
 		if (SourceSection.bRequired && SourceSection.Size <= 0)
 		{
-			UE_LOG(LogAssetHeaderPatcher, Error, TEXT("Unexpected section order for %s"), *SrcAsset);
+			UE_LOG(LogAssetHeaderPatcher, Error, TEXT("Unexpected missing required section for %s: %d is required but has zero size."),
+				*SrcAsset, SectionIdx);
 			return EResult::ErrorEmptyRequireSection;
 		}
 	}
@@ -2516,7 +2835,8 @@ FAssetHeaderPatcher::EResult FAssetHeaderPatcherInner::PatchHeader_WriteDestinat
 		if (LastSectionEndedAt)
 		{
 			int64 SizeToCopy = SourceSection.Offset - LastSectionEndedAt;
-			checkf(SizeToCopy >= 0, TEXT("Section %d of %s\n%zd -> %zd %zd"), SectionIdx, *SrcAsset, SourceSection.Offset, LastSectionEndedAt, SizeToCopy);
+			checkf(SizeToCopy >= 0, TEXT("Section %d of %s\n%" INT64_FMT " -> %" INT64_FMT" %" INT64_FMT),
+				SectionIdx, *SrcAsset, SourceSection.Offset, LastSectionEndedAt, SizeToCopy);
 			Writer.Serialize(SrcBuffer.GetData() + LastSectionEndedAt, SizeToCopy);
 		}
 		LastSectionEndedAt = SourceSection.Offset + SourceSection.Size;
@@ -2694,7 +3014,8 @@ FAssetHeaderPatcher::EResult FAssetHeaderPatcherInner::PatchHeader_WriteDestinat
 			const int64 Delta = ThumbnailTableSize - SourceSection.Size;
 			PatchSummaryOffsets(Summary, ThumbnailTableStartOffset, Delta);
 			Summary.TotalHeaderSize += (int32)Delta;
-			checkf(ThumbnailTableStartOffset == Summary.ThumbnailTableOffset, TEXT("%zd == %zd"), ThumbnailTableStartOffset, Summary.ThumbnailTableOffset);
+			checkf(ThumbnailTableStartOffset == Summary.ThumbnailTableOffset, TEXT("%" INT64_FMT " == %" INT64_FMT),
+				ThumbnailTableStartOffset, Summary.ThumbnailTableOffset);
 
 			break;
 		}
@@ -2702,11 +3023,17 @@ FAssetHeaderPatcher::EResult FAssetHeaderPatcherInner::PatchHeader_WriteDestinat
 		case EPatchedSection::AssetRegistryData:
 		{
 			const int64 AssetRegistryDataStartOffset = Writer.Tell();
-			checkf(AssetRegistryDataStartOffset == Summary.AssetRegistryDataOffset, TEXT("%zd == %zd"), AssetRegistryDataStartOffset, Summary.AssetRegistryDataOffset);
+			checkf(AssetRegistryDataStartOffset == Summary.AssetRegistryDataOffset, TEXT("%" INT64_FMT " == %" INT64_FMT),
+				AssetRegistryDataStartOffset, Summary.AssetRegistryDataOffset);
+
+			// Code to write this section copied and modified from UE::AssetRegistry::WritePackageData.
+			// TODO: Factor this into a public function in SavePackageUtilities and call that.
 
-			// Manually write this back out, there isn't a nicely factored function to call for this
 			if (AssetRegistryData.PkgData.DependencyDataOffset != INDEX_NONE)
 			{
+				// This field is conditionally written depending on package version and cookedness. If it is written it
+				// is guaranteed to != INDEX_NONE. We are not changing those properties for the package, so we write
+				// the field if and only if we found it present in the package originally.
 				Writer << AssetRegistryData.PkgData.DependencyDataOffset;
 			}
 			Writer << AssetRegistryData.PkgData.ObjectCount;
@@ -2732,13 +3059,60 @@ FAssetHeaderPatcher::EResult FAssetHeaderPatcherInner::PatchHeader_WriteDestinat
 			PatchSummaryOffsets(Summary, AssetRegistryDataStartOffset, Delta);
 			Summary.TotalHeaderSize += (int32)Delta;
 
-			if (AssetRegistryData.PkgData.DependencyDataOffset != INDEX_NONE)
+			break;
+		}
+
+		case EPatchedSection::AssetRegistryDependencyData:
+		{
+			int64 AssetRegistryDependencyDataStartOffset = Writer.Tell();
+
+			// Re-write the offset of this section into AssetRegistryData.PkgData.DependencyDataOffset in the
+			// earlier EPatchedSection::AssetRegistryData section.
+			Writer.Seek(Summary.AssetRegistryDataOffset);
+			Writer << AssetRegistryDependencyDataStartOffset;
+			if (Writer.IsError())
+			{
+				UE_LOG(LogAssetHeaderPatcher, Error, TEXT("Failed to write to %s"), *DstAsset);
+				return EResult::ErrorFailedToWriteToDestinationFile;
+			}
+			Writer.Seek(AssetRegistryDependencyDataStartOffset);
+			AssetRegistryData.PkgData.DependencyDataOffset = AssetRegistryDependencyDataStartOffset;
+
+			// Code to write this section copied and modified from UE::AssetRegistry::WritePackageData.
+			// TODO: Factor this into a public function in SavePackageUtilities and call that.
+
+			TBitArray<> ImportUsedInGameBits;
+			TBitArray<> SoftPackageUsedInGameBits;
+			ImportUsedInGameBits.Reserve(ImportTable.Num());
+			for (int32 ImportIndex = 0; ImportIndex < ImportTable.Num(); ++ImportIndex)
+			{
+				bool* UsedInGame = AssetRegistryData.ImportIndexUsedInGame.Find(ImportIndex);
+				bool bUsedInGame = UsedInGame ? *UsedInGame : true;
+				ImportUsedInGameBits.Add(bUsedInGame);
+			}
+			SoftPackageUsedInGameBits.Reserve(SoftPackageReferencesTable.Num());
+			for (int32 SoftPackageIndex = 0; SoftPackageIndex < SoftPackageReferencesTable.Num(); ++SoftPackageIndex)
+			{
+				bool* UsedInGame = AssetRegistryData.SoftPackageReferenceUsedInGame.Find(SoftPackageReferencesTable[SoftPackageIndex]);
+				bool bUsedInGame = UsedInGame ? *UsedInGame : true;
+				SoftPackageUsedInGameBits.Add(bUsedInGame);
+			}
+
+			Writer << ImportUsedInGameBits;
+			Writer << SoftPackageUsedInGameBits;
+
+			TArray<TPair<FName, uint32>> ExtraPackageDependenciesInts;
+			ExtraPackageDependenciesInts.Reserve(AssetRegistryData.ExtraPackageDependencies.Num());
+			for (const TPair<FName, UE::AssetRegistry::EExtraDependencyFlags>& DependencyPair : AssetRegistryData.ExtraPackageDependencies)
 			{
-				// DependencyDataOffset is not relative but points to just after the rest of the AR data
-				// We will seek back and write this later
-				const int64 DependencyDataDelta = AssetRegistryDataStartOffset - SourceSection.Offset + Delta;
-				AssetRegistryData.PkgData.DependencyDataOffset += DependencyDataDelta;
+				ExtraPackageDependenciesInts.Add({ DependencyPair.Key, static_cast<uint32>(DependencyPair.Value) });
 			}
+			Writer << ExtraPackageDependenciesInts;
+
+			const int64 AssetRegistryDependencyDataSize = Writer.Tell() - AssetRegistryDependencyDataStartOffset;
+			const int64 Delta = AssetRegistryDependencyDataSize - SourceSection.Size;
+			PatchSummaryOffsets(Summary, AssetRegistryDependencyDataStartOffset, Delta);
+			Summary.TotalHeaderSize += (int32)Delta;
 
 			break;
 		}
@@ -2757,7 +3131,8 @@ FAssetHeaderPatcher::EResult FAssetHeaderPatcherInner::PatchHeader_WriteDestinat
 
 	{   // copy the last blob
 		int64 SizeToCopy = SrcBuffer.Num() - LastSectionEndedAt;
-		checkf(SizeToCopy >= 0, TEXT("Section last of %s\n%zd -> %zd %zd"), *SrcAsset, SrcBuffer.Num(), LastSectionEndedAt, SizeToCopy);
+		checkf(SizeToCopy >= 0, TEXT("Section last of %s\n%" INT64_FMT " -> %" INT64_FMT " %" INT64_FMT),
+			*SrcAsset, SrcBuffer.Num(), LastSectionEndedAt, SizeToCopy);
 		Writer.Serialize(SrcBuffer.GetData() + LastSectionEndedAt, SizeToCopy);
 	}
 
@@ -2789,19 +3164,6 @@ FAssetHeaderPatcher::EResult FAssetHeaderPatcherInner::PatchHeader_WriteDestinat
 		return EResult::ErrorFailedToWriteToDestinationFile;
 	}
 
-	if (AssetRegistryData.PkgData.DependencyDataOffset != INDEX_NONE)
-	{
-		// Re-write asset registry dependency data offset
-		Writer.Seek(Summary.AssetRegistryDataOffset);
-		Writer << AssetRegistryData.PkgData.DependencyDataOffset;
-
-		if (Writer.IsError())
-		{
-			UE_LOG(LogAssetHeaderPatcher, Error, TEXT("Failed to write to %s"), *DstAsset);
-			return EResult::ErrorFailedToWriteToDestinationFile;
-		}
-	}
-
 	return EResult::Success;
 }
 
@@ -3180,8 +3542,43 @@ void FAssetHeaderPatcherInner::DumpState(FStringView OutputDirectory)
 	Builder.RemoveSuffix(1);
 	Builder.Append(TEXT("\n\t]\n"));
 
-	Builder.Append(TEXT("}"));
+	Builder.Append(TEXT("\t\"AssetRegistryDependencyData\":{ "));
+	{
+		Builder.Append(TEXT("\n\t\t\"ImportIndexUsedInGame\":{ "));
+		for (ImportIndex = 0; ImportIndex < ImportTable.Num(); ++ImportIndex)
+		{
+			bool* UsedInGame = AssetRegistryData.ImportIndexUsedInGame.Find(ImportIndex);
+			bool bUsedInGame = UsedInGame ? *UsedInGame : true;
+			Builder.Appendf(TEXT("\n\t\t\t%d : %s,"), ImportIndex, bUsedInGame ? TEXT("true") : TEXT("false"));
+		}
+		Builder.RemoveSuffix(1);
+		Builder.Append(TEXT("\n\t\t}"));
+
+		Builder.Append(TEXT(",\n\t\t\"SoftPackageReferenceUsedInGame\":{ "));
+		for (FName SoftPackageReference : SoftPackageReferencesTable)
+		{
+			bool* UsedInGame = AssetRegistryData.SoftPackageReferenceUsedInGame.Find(SoftPackageReference);
+			bool bUsedInGame = UsedInGame ? *UsedInGame : true;
+			Builder.Append(TEXT("\n\t\t\t"));
+			Builder << SoftPackageReference;
+			Builder.Appendf(TEXT(" : %s,"), bUsedInGame ? TEXT("true") : TEXT("false"));
+		}
+		Builder.RemoveSuffix(1);
+		Builder.Append(TEXT("\n\t\t}"));
+
+		Builder.Append(TEXT(",\n\t\t\"ExtraPackageDependencies\":[ "));
+		for (const TPair<FName, UE::AssetRegistry::EExtraDependencyFlags>& Pair : AssetRegistryData.ExtraPackageDependencies)
+		{
+			Builder.Append(TEXT("\n\t\t\t[ \""));
+			Builder << Pair.Key;
+			Builder.Appendf(TEXT("\", 0x%x],"), (uint32) Pair.Value);
+		}
+		Builder.RemoveSuffix(1);
+		Builder.Append(TEXT("\n\t\t]"));
+	}
+	Builder.Append(TEXT("\n\t},\n"));
 
+	Builder.Append(TEXT("}"));
 
 	// Write to disk
 	TStringBuilder<256> OutPath;
@@ -3329,101 +3726,126 @@ TEST_CASE_NAMED(FAssetHeaderPatcherTests, "AssetHeaderPatcher", "[AssetHeaderPat
 			return Export;
 		};
 
-	TArray<TPair<FObjectImport, FObjectImport>> ImportTestCases
+	struct FImportTestCase
+	{
+		FObjectImport Src;
+		FObjectImport Dst;
+		bool bExistingImport;
+	};
+	TArray<FImportTestCase> ImportTestCases
 	{
 		// /Script/SrcEngineModule																-> (remains unchanged)
 		{	
-			MakeImport(SrcEngineModuleImportObjectName, FPackageIndex(), SrcImportClassPackage, SrcImportClassName, SrcPackagePathFName), 
-			MakeImport(SrcEngineModuleImportObjectName, FPackageIndex(), DstImportClassPackage, DstImportClassName, DstPackagePathFName) 
+			.Src = MakeImport(SrcEngineModuleImportObjectName, FPackageIndex(), SrcImportClassPackage, SrcImportClassName, SrcPackagePathFName), 
+			.Dst = MakeImport(SrcEngineModuleImportObjectName, FPackageIndex(), DstImportClassPackage, DstImportClassName, DstPackagePathFName),
+			.bExistingImport = true,
 		},
 		// /Script/SrcEngineModule.Default__SrcTypeA											-> /Script/DstEngineModule.Default__DstTypeA
 		{ 
-			MakeImport(Default__SrcTypeAImportObjectName, FPackageIndex::FromImport(0), SrcImportClassPackage, SrcImportClassName, SrcPackagePathFName), 
-			MakeImport(Default__DstTypeAImportObjectName, FPackageIndex::FromImport(13), DstImportClassPackage, DstImportClassName, DstPackagePathFName) 
+			.Src = MakeImport(Default__SrcTypeAImportObjectName, FPackageIndex::FromImport(0), SrcImportClassPackage, SrcImportClassName, SrcPackagePathFName), 
+			.Dst = MakeImport(Default__DstTypeAImportObjectName, FPackageIndex::FromImport(13), DstImportClassPackage, DstImportClassName, DstPackagePathFName),
+			.bExistingImport = true,
 		},
 		// /Script/SrcEngineModule.Default__SrcTypeA.SrcPropertyToChangeA						-> /Script/DstEngineModule.Default__DstTypeA.DstPropertyToChangeA
 		{ 
-			MakeImport(SrcPropertyToChangeAImportObjectName, FPackageIndex::FromImport(1), SrcImportClassPackage, SrcImportClassName, SrcPackagePathFName), 
-			MakeImport(DstPropertyToChangeAImportObjectName, FPackageIndex::FromImport(1), DstImportClassPackage, DstImportClassName, DstPackagePathFName) 
+			.Src = MakeImport(SrcPropertyToChangeAImportObjectName, FPackageIndex::FromImport(1), SrcImportClassPackage, SrcImportClassName, SrcPackagePathFName), 
+			.Dst = MakeImport(DstPropertyToChangeAImportObjectName, FPackageIndex::FromImport(1), DstImportClassPackage, DstImportClassName, DstPackagePathFName),
+			.bExistingImport = true,
 		},
 		// /Script/SrcEngineModule.Default__SrcTypeA.MovedButNotRenamedProperty					-> /Script/DstEngineModule.Default__DstTypeA.MovedButNotRenamedProperty
 		{ 
-			MakeImport(MovedButNotRenamedPropertyImportObjectName, FPackageIndex::FromImport(1), SrcImportClassPackage, SrcImportClassName, SrcPackagePathFName), 
-			MakeImport(MovedButNotRenamedPropertyImportObjectName, FPackageIndex::FromImport(1), DstImportClassPackage, DstImportClassName, DstPackagePathFName) 
+			.Src = MakeImport(MovedButNotRenamedPropertyImportObjectName, FPackageIndex::FromImport(1), SrcImportClassPackage, SrcImportClassName, SrcPackagePathFName), 
+			.Dst = MakeImport(MovedButNotRenamedPropertyImportObjectName, FPackageIndex::FromImport(1), DstImportClassPackage, DstImportClassName, DstPackagePathFName), 
+			.bExistingImport = true,
 		},
 		// /Script/SrcEngineModule.Default__SrcTypeA.MovedButNotRenamedProperty.InnerMovedButNotRenamedProperty
 		// -> 
 		// /Script/DstEngineModule.Default__DstTypeA.MovedButNotRenamedProperty.InnerMovedButNotRenamedProperty
 		{
-			MakeImport(InnerMovedButNotRenamedPropertyImportObjectName, FPackageIndex::FromImport(3), SrcImportClassPackage, SrcImportClassName, SrcPackagePathFName),
-			MakeImport(InnerMovedButNotRenamedPropertyImportObjectName, FPackageIndex::FromImport(3), DstImportClassPackage, DstImportClassName, DstPackagePathFName)
+			.Src = MakeImport(InnerMovedButNotRenamedPropertyImportObjectName, FPackageIndex::FromImport(3), SrcImportClassPackage, SrcImportClassName, SrcPackagePathFName),
+			.Dst = MakeImport(InnerMovedButNotRenamedPropertyImportObjectName, FPackageIndex::FromImport(3), DstImportClassPackage, DstImportClassName, DstPackagePathFName),
+			.bExistingImport = true,
 		},
 		// /Script/SrcEngineModule.Default__SrcTypeA.MovedButNotRenamedProperty.InnerMovedButNotRenamedProperty.InnerInnerMovedButNotRenamedProperty 
 		// -> 
 		// /Script/DstEngineModule.Default__DstTypeA.MovedButNotRenamedProperty.InnerMovedButNotRenamedProperty.InnerInnerMovedButNotRenamedProperty 
 		{
-			MakeImport(InnerInnerMovedButNotRenamedPropertyImportObjectName, FPackageIndex::FromImport(4), SrcImportClassPackage, SrcImportClassName, SrcPackagePathFName),
-			MakeImport(InnerInnerMovedButNotRenamedPropertyImportObjectName, FPackageIndex::FromImport(4), DstImportClassPackage, DstImportClassName, DstPackagePathFName)
+			.Src = MakeImport(InnerInnerMovedButNotRenamedPropertyImportObjectName, FPackageIndex::FromImport(4), SrcImportClassPackage, SrcImportClassName, SrcPackagePathFName),
+			.Dst = MakeImport(InnerInnerMovedButNotRenamedPropertyImportObjectName, FPackageIndex::FromImport(4), DstImportClassPackage, DstImportClassName, DstPackagePathFName),
+			.bExistingImport = true,
 		},
 		// /Script/SrcEngineModule.Default__OnlySubTypeChanged									-> (remains unchanged)
 		{ 
-			MakeImport(Default__OnlySubTypeChangedImportObjectName, FPackageIndex::FromImport(0), SrcImportClassPackage, SrcImportClassName, SrcPackagePathFName), 
-			MakeImport(Default__OnlySubTypeChangedImportObjectName, FPackageIndex::FromImport(0), DstImportClassPackage, DstImportClassName, DstPackagePathFName)
+			.Src = MakeImport(Default__OnlySubTypeChangedImportObjectName, FPackageIndex::FromImport(0), SrcImportClassPackage, SrcImportClassName, SrcPackagePathFName), 
+			.Dst = MakeImport(Default__OnlySubTypeChangedImportObjectName, FPackageIndex::FromImport(0), DstImportClassPackage, DstImportClassName, DstPackagePathFName),
+			.bExistingImport = true,
 		},
 		// /Script/SrcEngineModule.Default__OnlySubTypeChanged.SrcPropertyToChangeB				-> /Script/SrcEngineModule.Default__OnlySubTypeChanged.DstPropertyToChangeB
 		{ 
-			MakeImport(SrcPropertyToChangeBImportObjectName, FPackageIndex::FromImport(6), SrcImportClassPackage, SrcImportClassName, SrcPackagePathFName), 
-			MakeImport(DstPropertyToChangeBImportObjectName, FPackageIndex::FromImport(6), DstImportClassPackage, DstImportClassName, DstPackagePathFName) 
+			.Src = MakeImport(SrcPropertyToChangeBImportObjectName, FPackageIndex::FromImport(6), SrcImportClassPackage, SrcImportClassName, SrcPackagePathFName), 
+			.Dst = MakeImport(DstPropertyToChangeBImportObjectName, FPackageIndex::FromImport(6), DstImportClassPackage, DstImportClassName, DstPackagePathFName),
+			.bExistingImport = true,
 		},
 		// /Script/SrcEngineModule.Default__OnlySubTypeChanged.UnchangedProperty				-> (remains unchanged)
 		{ 
-			MakeImport(UnchangedPropertyImportObjectName, FPackageIndex::FromImport(6), SrcImportClassPackage, SrcImportClassName, SrcPackagePathFName), 
-			MakeImport(UnchangedPropertyImportObjectName, FPackageIndex::FromImport(6), DstImportClassPackage, DstImportClassName, DstPackagePathFName) 
+			.Src = MakeImport(UnchangedPropertyImportObjectName, FPackageIndex::FromImport(6), SrcImportClassPackage, SrcImportClassName, SrcPackagePathFName), 
+			.Dst = MakeImport(UnchangedPropertyImportObjectName, FPackageIndex::FromImport(6), DstImportClassPackage, DstImportClassName, DstPackagePathFName),
+			.bExistingImport = true,
 		},
 		// /Script/SrcEngineModule.Default__SrcTypeB											-> /Script/SrcEngineModule.Default__DstTypeB
 		{ 
-			MakeImport(Default__SrcTypeBImportObjectName, FPackageIndex::FromImport(0), SrcImportClassPackage, SrcImportClassName, SrcPackagePathFName), 
-			MakeImport(Default__DstTypeBImportObjectName, FPackageIndex::FromImport(0), DstImportClassPackage, DstImportClassName, DstPackagePathFName) 
+			.Src = MakeImport(Default__SrcTypeBImportObjectName, FPackageIndex::FromImport(0), SrcImportClassPackage, SrcImportClassName, SrcPackagePathFName), 
+			.Dst = MakeImport(Default__DstTypeBImportObjectName, FPackageIndex::FromImport(0), DstImportClassPackage, DstImportClassName, DstPackagePathFName),
+			.bExistingImport = true,
 		},
 		// /Script/SrcEngineModule.Default__SrcTypeB.MovedButNotRenamedProperty					-> /Script/SrcEngineModule.Default__DstTypeB.MovedButNotRenamedProperty
 		{ 
-			MakeImport(MovedButNotRenamedPropertyImportObjectName, FPackageIndex::FromImport(9), SrcImportClassPackage, SrcImportClassName, SrcPackagePathFName), 
-			MakeImport(MovedButNotRenamedPropertyImportObjectName, FPackageIndex::FromImport(9), DstImportClassPackage, DstImportClassName, DstPackagePathFName) 
+			.Src = MakeImport(MovedButNotRenamedPropertyImportObjectName, FPackageIndex::FromImport(9), SrcImportClassPackage, SrcImportClassName, SrcPackagePathFName), 
+			.Dst = MakeImport(MovedButNotRenamedPropertyImportObjectName, FPackageIndex::FromImport(9), DstImportClassPackage, DstImportClassName, DstPackagePathFName),
+			.bExistingImport = true,
 		},
 		// /Script/SrcEngineModule.Default__MovedButNotRenamedType								-> /Script/DstEngineModule.Default__MovedButNotRenamedType
 		{ 
-			MakeImport(Default__MovedButNotRenamedTypeImportObjectName, FPackageIndex::FromImport(0), SrcImportClassPackage, SrcImportClassName, SrcPackagePathFName), 
-			MakeImport(Default__MovedButNotRenamedTypeImportObjectName, FPackageIndex::FromImport(13), DstImportClassPackage, DstImportClassName, DstPackagePathFName) 
+			.Src = MakeImport(Default__MovedButNotRenamedTypeImportObjectName, FPackageIndex::FromImport(0), SrcImportClassPackage, SrcImportClassName, SrcPackagePathFName), 
+			.Dst = MakeImport(Default__MovedButNotRenamedTypeImportObjectName, FPackageIndex::FromImport(13), DstImportClassPackage, DstImportClassName, DstPackagePathFName),
+			.bExistingImport = true,
 		},
 		// /Script/SrcEngineModule.Default__MovedButNotRenamedType.MovedButNotRenamedProperty	-> /Script/DstEngineModule.Default__MovedButNotRenamedType.MovedButNotRenamedProperty
 		{ 
-			MakeImport(MovedButNotRenamedPropertyImportObjectName, FPackageIndex::FromImport(11), SrcImportClassPackage, SrcImportClassName, SrcPackagePathFName), 
-			MakeImport(MovedButNotRenamedPropertyImportObjectName, FPackageIndex::FromImport(11), DstImportClassPackage, DstImportClassName, DstPackagePathFName) 
+			.Src = MakeImport(MovedButNotRenamedPropertyImportObjectName, FPackageIndex::FromImport(11), SrcImportClassPackage, SrcImportClassName, SrcPackagePathFName), 
+			.Dst = MakeImport(MovedButNotRenamedPropertyImportObjectName, FPackageIndex::FromImport(11), DstImportClassPackage, DstImportClassName, DstPackagePathFName),
+			.bExistingImport = true,
 		},
 		// <none>																				-> /Script/DstEngineModule
 		{ 
-			FObjectImport(), 
-			MakeImport(DstEngineModuleImportObjectName, FPackageIndex(), DstImportClassPackage, DstImportClassName, DstPackagePathFName)
+			.Src = FObjectImport(), 
+			.Dst = MakeImport(DstEngineModuleImportObjectName, FPackageIndex(), DstImportClassPackage, DstImportClassName, DstPackagePathFName),
+			.bExistingImport = false,
 		},
 	};
 
-	TArray<TPair<FObjectExport, FObjectExport>> ExportTestCases
+	struct FExportTestCase
+	{
+		FObjectExport Src;
+		FObjectExport Dst;
+	};
+	TArray<FExportTestCase> ExportTestCases
 	{
 		// SrcPackage																			-> DstPackage
 		// Super: /Script/SrcEngineModule														-> /Script/DstEngineModule
 		// Class: /Script/SrcEngineModule.Default__SrcTypeA.SrcPropertyToChangeA				-> /Script/DstEngineModule.Default__DstTypeA.DstPropertyToChangeA
 		// Template: /Script/SrcEngineModule.Default__OnlySubTypeChanged						-> /Script/SrcEngineModule.Default__OnlySubTypeChanged
 		{
-			MakeExport(SrcExportObjectFName, FPackageIndex::FromExport(0), FPackageIndex(), FPackageIndex::FromImport(0), FPackageIndex::FromImport(2), FPackageIndex::FromImport(4)),
-			MakeExport(DstExportObjectFName, FPackageIndex::FromExport(0), FPackageIndex(), FPackageIndex::FromImport(11), FPackageIndex::FromImport(2), FPackageIndex::FromImport(4))
+			.Src = MakeExport(SrcExportObjectFName, FPackageIndex::FromExport(0), FPackageIndex(), FPackageIndex::FromImport(0), FPackageIndex::FromImport(2), FPackageIndex::FromImport(4)),
+			.Dst = MakeExport(DstExportObjectFName, FPackageIndex::FromExport(0), FPackageIndex(), FPackageIndex::FromImport(11), FPackageIndex::FromImport(2), FPackageIndex::FromImport(4))
 		},
 		// SrcPackage.SrcPackage																-> DstPackage.DstPackage
 		// Super: /Script/SrcEngineModule.Default__SrcTypeA										-> /Script/DstEngineModule.Default__DstTypeA
 		// Class: /Script/SrcEngineModule.Default__SrcTypeA.MovedButNotRenamedProperty			-> /Script/DstEngineModule.Default__DstTypeA.MovedButNotRenamedProperty
 		// Template: /Script/SrcEngineModule.Default__OnlySubTypeChanged.SrcPropertyToChangeB	-> /Script/SrcEngineModule.Default__OnlySubTypeChanged.DstPropertyToChangeB
 		{
-			MakeExport(SrcExportObjectFName, FPackageIndex::FromExport(1), FPackageIndex::FromExport(0), FPackageIndex::FromImport(1), FPackageIndex::FromImport(3), FPackageIndex::FromImport(5)),
-			MakeExport(DstExportObjectFName, FPackageIndex::FromExport(1), FPackageIndex::FromExport(0), FPackageIndex::FromImport(1), FPackageIndex::FromImport(3), FPackageIndex::FromImport(5))
+			.Src = MakeExport(SrcExportObjectFName, FPackageIndex::FromExport(1), FPackageIndex::FromExport(0), FPackageIndex::FromImport(1), FPackageIndex::FromImport(3), FPackageIndex::FromImport(5)),
+			.Dst = MakeExport(DstExportObjectFName, FPackageIndex::FromExport(1), FPackageIndex::FromExport(0), FPackageIndex::FromImport(1), FPackageIndex::FromImport(3), FPackageIndex::FromImport(5))
 		}
 	};
 
@@ -3504,9 +3926,9 @@ TEST_CASE_NAMED(FAssetHeaderPatcherTests, "AssetHeaderPatcher", "[AssetHeaderPat
 				};
 
 			// Import/Export Table (see breakdown of import names for the intended tests in the FObjectResource test section below)
-			for (const TPair<FObjectImport, FObjectImport>& TestCase : ImportTestCases)
+			for (const FImportTestCase& TestCase : ImportTestCases)
 			{
-				const FObjectImport& Import = TestCase.Key;
+				const FObjectImport& Import = TestCase.Src;
 				if (Import.ObjectName == NAME_None)
 				{
 					// We have more cases than initial starting states so if we see an empty
@@ -3523,9 +3945,9 @@ TEST_CASE_NAMED(FAssetHeaderPatcherTests, "AssetHeaderPatcher", "[AssetHeaderPat
 				Patcher.ImportTable.Add(Import);
 			}
 
-			for (const TPair<FObjectExport, FObjectExport>& TestCase : ExportTestCases)
+			for (const FExportTestCase& TestCase : ExportTestCases)
 			{
-				const FObjectExport& Export = TestCase.Key;
+				const FObjectExport& Export = TestCase.Src;
 				CheckNameTableInit(Export.ObjectName);
 #if WITH_EDITORONLY_DATA
 				CheckNameTableInit(Export.OldClassName);
@@ -4064,7 +4486,7 @@ TEST_CASE_NAMED(FAssetHeaderPatcherTests, "AssetHeaderPatcher", "[AssetHeaderPat
 			check(Patcher.ImportTable.Num() <= ImportTestCases.Num());
 			for (int32 i = 0; i < Patcher.ImportTable.Num(); ++i)
 			{
-				const FObjectImport& Expected = ImportTestCases[i].Key;
+				const FObjectImport& Expected = ImportTestCases[i].Src;
 				const FObjectImport& Actual = Patcher.ImportTable[i];
 
 				CheckFNames(Expected.ObjectName, Actual.ObjectName);
@@ -4078,7 +4500,7 @@ TEST_CASE_NAMED(FAssetHeaderPatcherTests, "AssetHeaderPatcher", "[AssetHeaderPat
 			check(Patcher.ExportTable.Num() <= ExportTestCases.Num());
 			for (int32 i = 0; i < Patcher.ExportTable.Num(); ++i)
 			{
-				const FObjectExport& Expected = ExportTestCases[i].Key;
+				const FObjectExport& Expected = ExportTestCases[i].Src;
 				const FObjectExport& Actual = Patcher.ExportTable[i];
 
 				CheckFNames(Expected.ObjectName, Actual.ObjectName);
@@ -4096,7 +4518,8 @@ TEST_CASE_NAMED(FAssetHeaderPatcherTests, "AssetHeaderPatcher", "[AssetHeaderPat
 			TArray<FAssetHeaderPatcherInner::FImportPatch> ImportPatches;
 			Patcher.GetExportTablePatches(ExportPatches);
 			CHECK(!ExportPatches.IsEmpty());
-			Patcher.GetImportTablePatches(ImportPatches, NewImportCount);
+			FAssetHeaderPatcher::EResult Result = Patcher.GetImportTablePatches(ImportPatches, NewImportCount);
+			CHECK(Result == FAssetHeaderPatcher::EResult::Success);
 			CHECK(!ExportPatches.IsEmpty());
 			Patcher.PatchExportAndImportTables(ExportPatches, ImportPatches, NewImportCount);
 			Patcher.PatchNameTable();
@@ -4105,24 +4528,33 @@ TEST_CASE_NAMED(FAssetHeaderPatcherTests, "AssetHeaderPatcher", "[AssetHeaderPat
 			CHECK(Patcher.ImportTable.Num() == ImportTestCases.Num());
 			for (int32 i = 0; i < ImportTestCases.Num(); ++i)
 			{
-				const FObjectImport& Expected = ImportTestCases[i].Value;
+				const FObjectImport& Expected = ImportTestCases[i].Dst;
 				const FObjectImport& Actual = Patcher.ImportTable[i];
 
 				CheckFNames(Expected.ObjectName, Actual.ObjectName);
 				CHECK(Expected.OuterIndex == Actual.OuterIndex);
-				CheckFNames(Expected.ClassName, Actual.ClassName);
-				CheckFNames(Expected.ClassPackage, Actual.ClassPackage);
+				if (ImportTestCases[i].bExistingImport)
+				{
+					CheckFNames(Expected.ClassName, Actual.ClassName);
+					CheckFNames(Expected.ClassPackage, Actual.ClassPackage);
 #if WITH_EDITORONLY_DATA
-				CheckFNames(Expected.PackageName, Actual.PackageName);
-				CheckFNames(Expected.OldClassName, Actual.OldClassName);
-				CheckFNames(Actual.OldClassName, NAME_None);
+					CheckFNames(Expected.PackageName, Actual.PackageName);
+					CheckFNames(Expected.OldClassName, Actual.OldClassName);
+					CheckFNames(Actual.OldClassName, NAME_None);
 #endif
+				}
+				else
+				{
+					// For new imports created by the patcher, we don't yet have a contract for what they
+					// should report for Class of the import and external packagename, because it would have
+					// to read that out of the target packages.
+				}
 			}
 			
 			CHECK(Patcher.ExportTable.Num() == ExportTestCases.Num());
 			for (int32 i = 0; i < Patcher.ExportTable.Num(); ++i)
 			{
-				const FObjectExport& Expected = ExportTestCases[i].Value;
+				const FObjectExport& Expected = ExportTestCases[i].Dst;
 				const FObjectExport& Actual = Patcher.ExportTable[i];
 
 				CheckFNames(Expected.ObjectName, Actual.ObjectName);
