AssetTools: Add API to Asset Tools to register/unregister default names for asset types to override the default name provided by the asset factory.
Modify CreateAssetsFrom() and CreateAssetWithDialog() to check for registered default names.
Modify content browser plugin to check for registered default names when creating new assets in OnNewAssetRequested()

diff --git a/Engine/Plugins/Editor/ContentBrowser/ContentBrowserAssetDataSource/Source/ContentBrowserAssetDataSource/Private/ContentBrowserAssetDataSource.cpp b/Engine/Plugins/Editor/ContentBrowser/ContentBrowserAssetDataSource/Source/ContentBrowserAssetDataSource/Private/ContentBrowserAssetDataSource.cpp
--- a/Engine/Plugins/Editor/ContentBrowser/ContentBrowserAssetDataSource/Source/ContentBrowserAssetDataSource/Private/ContentBrowserAssetDataSource.cpp
+++ b/Engine/Plugins/Editor/ContentBrowser/ContentBrowserAssetDataSource/Source/ContentBrowserAssetDataSource/Private/ContentBrowserAssetDataSource.cpp
@@ -3856,11 +3856,14 @@ void UContentBrowserAssetDataSource::OnNewAssetRequested(const FName InPath, TWe
 		{
 			FEditorDelegates::OnNewAssetCreated.Broadcast(NewFactory);
 
-			FString DefaultAssetName;
+			const TOptional<FString> DefaultAssetName = AssetTools->GetDefaultAssetNameForClass(NewFactory->GetSupportedClass());
+			const FString BaseAssetName = DefaultAssetName.IsSet() ? DefaultAssetName.GetValue() : NewFactory->GetDefaultNewAssetName();
+
+			FString UniqueAssetName;
 			FString PackageNameToUse;
-			AssetTools->CreateUniqueAssetName(InPath.ToString() / NewFactory->GetDefaultNewAssetName(), FString(), PackageNameToUse, DefaultAssetName);
+			AssetTools->CreateUniqueAssetName(InPath.ToString() / BaseAssetName, FString(), PackageNameToUse, UniqueAssetName);
 
-			OnBeginCreateAsset(*DefaultAssetName, InPath, NewFactory->GetSupportedClass(), NewFactory, InOnBeginItemCreation);
+			OnBeginCreateAsset(*UniqueAssetName, InPath, NewFactory->GetSupportedClass(), NewFactory, InOnBeginItemCreation);
 		}
 	}
 }
diff --git a/Engine/Source/Developer/AssetTools/Private/AssetTools.cpp b/Engine/Source/Developer/AssetTools/Private/AssetTools.cpp
--- a/Engine/Source/Developer/AssetTools/Private/AssetTools.cpp
+++ b/Engine/Source/Developer/AssetTools/Private/AssetTools.cpp
@@ -1783,9 +1783,16 @@ void UAssetToolsImpl::CreateAssetsFrom(TConstArrayView<UObject*> SourceObjects,
 			if (UFactory* Factory = FactoryConstructor(SourceObject))
 			{
 				// Create an appropriate and unique name 
+				const TOptional<FString> DefaultAssetName = GetDefaultAssetNameForClass(CreateAssetType);
+
+				const FString BasePackageName = DefaultAssetName.IsSet() ? DefaultAssetName.GetValue() : SourceObject->GetName();
+				const FString Suffix = DefaultAssetName.IsSet() ? FString() : DefaultSuffix;
+
+				const FString PackagePath = FPackageName::GetLongPackagePath(SourceObject->GetPackage()->GetName());
+
 				FString Name;
 				FString PackageName;
-				CreateUniqueAssetName(SourceObject->GetOutermost()->GetName(), DefaultSuffix, PackageName, Name);
+				CreateUniqueAssetName(PackagePath / BasePackageName, Suffix, PackageName, Name);
 
 				FContentBrowserModule& ContentBrowserModule = FModuleManager::LoadModuleChecked<FContentBrowserModule>("ContentBrowser");
 				ContentBrowserModule.Get().CreateNewAsset(Name, FPackageName::GetLongPackagePath(PackageName), CreateAssetType, Factory);
@@ -1803,9 +1810,16 @@ void UAssetToolsImpl::CreateAssetsFrom(TConstArrayView<UObject*> SourceObjects,
 				if (UFactory* Factory = FactoryConstructor(SourceObject))
 				{
 					// Determine an appropriate name
+					const TOptional<FString> DefaultAssetName = GetDefaultAssetNameForClass(CreateAssetType);
+
+					const FString BasePackageName = DefaultAssetName.IsSet() ? DefaultAssetName.GetValue() : SourceObject->GetName();
+					const FString Suffix = DefaultAssetName.IsSet() ? FString() : DefaultSuffix;
+
+					const FString PackagePath = FPackageName::GetLongPackagePath(SourceObject->GetPackage()->GetName());
+
 					FString Name;
 					FString PackageName;
-					CreateUniqueAssetName(SourceObject->GetOutermost()->GetName(), DefaultSuffix, PackageName, Name);
+					CreateUniqueAssetName(PackagePath / BasePackageName, Suffix, PackageName, Name);
 					
 					if (UObject* NewAsset = CreateAsset(Name, FPackageName::GetLongPackagePath(PackageName), CreateAssetType, Factory, CallingContext))
 					{
@@ -1836,9 +1850,13 @@ UObject* UAssetToolsImpl::CreateAssetWithDialog(UClass* AssetClass, UFactory* Fa
 			AssetPath = TEXT("/Game");
 		}
 
+		// Determine an appropriate name
+		const TOptional<FString> DefaultAssetName = GetDefaultAssetNameForClass(AssetClass);
+		const FString BaseAssetName = DefaultAssetName.IsSet() ? DefaultAssetName.GetValue() : Factory->GetDefaultNewAssetName();
+
 		FString PackageName;
 		FString AssetName;
-		CreateUniqueAssetName(AssetPath / Factory->GetDefaultNewAssetName(), TEXT(""), PackageName, AssetName);
+		CreateUniqueAssetName(AssetPath / BaseAssetName, TEXT(""), PackageName, AssetName);
 
 		return CreateAssetWithDialog(AssetName, AssetPath, AssetClass, Factory, CallingContext);
 	}
@@ -6385,6 +6403,27 @@ void UAssetToolsImpl::UnregisterSanitizeNameDelegate(const FName OwnerName)
 	SanitizeNameDelegates.Remove(OwnerName);
 }
 
+TOptional<FString> UAssetToolsImpl::GetDefaultAssetNameForClass(const UClass* Class) const
+{
+	if (const FString* AssetName = DefaultNewAssetNames.Find(Class))
+	{
+		return *AssetName;
+	}
+
+	return TOptional<FString>();
+}
+
+void UAssetToolsImpl::RegisterDefaultAssetNameForClass(const UClass* Class, const FString& DefaultName)
+{
+	FString SantizedDefaultName = ObjectTools::SanitizeInvalidChars(DefaultName, INVALID_LONGPACKAGE_CHARACTERS INVALID_OBJECTNAME_CHARACTERS);
+	DefaultNewAssetNames.Add(Class, SantizedDefaultName);
+}
+
+void UAssetToolsImpl::UnregisterDefaultAssetNameForClass(const UClass* Class)
+{
+	DefaultNewAssetNames.Remove(Class);
+}
+
 void UAssetToolsImpl::RegisterCanMigrateAsset(const FName OwnerName, UE::AssetTools::FCanMigrateAsset Delegate)
 {
 	CanMigrateAssetDelegates.Add(OwnerName, MoveTemp(Delegate));
diff --git a/Engine/Source/Developer/AssetTools/Private/AssetTools.h b/Engine/Source/Developer/AssetTools/Private/AssetTools.h
--- a/Engine/Source/Developer/AssetTools/Private/AssetTools.h
+++ b/Engine/Source/Developer/AssetTools/Private/AssetTools.h
@@ -157,6 +157,9 @@ public:
 	virtual bool SanitizeName(FString& NameToSanitize) override;
 	virtual void RegisterSanitizeNameDelegate(const FName OwnerName, FSanitizeName Delegate) override;
 	virtual void UnregisterSanitizeNameDelegate(const FName OwnerName) override;
+	virtual TOptional<FString> GetDefaultAssetNameForClass(const UClass* Class) const override;
+	virtual void RegisterDefaultAssetNameForClass(const UClass* Class, const FString& DefaultName) override;
+	virtual void UnregisterDefaultAssetNameForClass(const UClass* Class) override;
 	virtual void RegisterCanMigrateAsset(const FName OwnerName, UE::AssetTools::FCanMigrateAsset Delegate) override;
 	virtual void UnregisterCanMigrateAsset(const FName OwnerName) override;
 	virtual bool CanAssetBePublic(FStringView AssetPath) const override;
@@ -279,6 +282,9 @@ private:
 	/** List of sub content paths denied for every mount. */
 	TArray<FString> SubContentDenyListPaths;
 
+	/** Default names to use when creating new assets */
+	TMap<const UClass*, FString> DefaultNewAssetNames;
+
 	bool CreateAssetsAsExternallyReferenceable;
 	UE::AssetTools::FShouldCreateAssetsAsExternallyReferenceableForPath CreateAssetsAsExternallyReferenceableForPathDelegate;
 
diff --git a/Engine/Source/Developer/AssetTools/Public/IAssetTools.h b/Engine/Source/Developer/AssetTools/Public/IAssetTools.h
--- a/Engine/Source/Developer/AssetTools/Public/IAssetTools.h
+++ b/Engine/Source/Developer/AssetTools/Public/IAssetTools.h
@@ -680,6 +680,13 @@ public:
 	/** Remove a previously-set global name filter */
 	virtual void UnregisterSanitizeNameDelegate(const FName OwnerName) = 0;
 
+	/** Get the default name for the supplied class (if one is registered) to use when creating new assets of that type */
+	virtual TOptional<FString> GetDefaultAssetNameForClass(const UClass* Class) const = 0;
+	/** Register a default name for the supplied class to use when creating new assets of that type */
+	virtual void RegisterDefaultAssetNameForClass(const UClass* Class, const FString& DefaultName) = 0;
+	/** Remove any previously registered name for the supplied class */
+	virtual void UnregisterDefaultAssetNameForClass(const UClass* Class) = 0;
+
 	/** Show notification that writable folder filter blocked an action */
 	virtual void NotifyBlockedByWritableFolderFilter() const = 0;
 
