[AssetHeaderPatcher]

- Allow SoftObjectPaths to change in size when patching. This can occur with FString changes to serialization (e.g. if we load a string wchar_t but then resave using UTF8) but as well, it's possible a redirect changes the sub-object path in a Softpath which would alter the length of the FString serialized in SoftPaths as part of patching. The original code assumed we would only ever patch FNames and thus the FString would not change however that is not accurate,

#rnx

diff --git a/Engine/Source/Developer/AssetTools/Private/AssetHeaderPatcher.cpp b/Engine/Source/Developer/AssetTools/Private/AssetHeaderPatcher.cpp
--- a/Engine/Source/Developer/AssetTools/Private/AssetHeaderPatcher.cpp
+++ b/Engine/Source/Developer/AssetTools/Private/AssetHeaderPatcher.cpp
@@ -2901,7 +2901,8 @@ FAssetHeaderPatcher::EResult FAssetHeaderPatcherInner::PatchHeader_WriteDestinat
 
 			const int64 TableSize = Writer.Tell() - TableStartOffset;
 			const int64 Delta = TableSize - SourceSection.Size;
-			checkf(Delta == 0, TEXT("Delta should be Zero. is %d"), (int)Delta);
+			PatchSummaryOffsets(Summary, TableStartOffset, Delta);
+			Summary.TotalHeaderSize += (int32)Delta;
 			check(Summary.SoftObjectPathsCount == SoftObjectPathTable.Num());
 			check(Summary.SoftObjectPathsOffset == TableStartOffset);
 
