[AssetHeaderPatcher] Fix thumbnails patching not properly computing thumbnails offsets.

diff --git a/Engine/Source/Developer/AssetTools/Private/AssetHeaderPatcher.cpp b/Engine/Source/Developer/AssetTools/Private/AssetHeaderPatcher.cpp
--- a/Engine/Source/Developer/AssetTools/Private/AssetHeaderPatcher.cpp
+++ b/Engine/Source/Developer/AssetTools/Private/AssetHeaderPatcher.cpp
@@ -852,7 +852,6 @@ public:
 		FString ObjectShortClassName;
 		FString ObjectPathWithoutPackageName;
 		int32 FileOffset = 0;
-		int32 Delta = 0;
 	};
 
 	FAssetHeaderPatcherInner(const FString& InSrcAsset, const FString& InDstAsset, const TMap<FString, FString>& InStringReplacements, const TMap<FString, FString>& InStringMountPointReplacements, FArchive* InDstArchive = nullptr)
@@ -2479,15 +2478,8 @@ bool FAssetHeaderPatcherInner::DoPatch(FThumbnailEntry& InThumbnailEntry)
 
 	if (bPatched)
 	{
-		// Since we patched, we will cause the inline string name to affect the thumbnail offsets. Calculate the size change
-		// here so we can use it during writing where we will fix up the offsets.
-		int32 Delta = -(InThumbnailEntry.ObjectShortClassName.Len() + InThumbnailEntry.ObjectPathWithoutPackageName.Len());
-
 		InThumbnailEntry.ObjectShortClassName = RedirectedClassName.ObjectName.ToString();
 		InThumbnailEntry.ObjectPathWithoutPackageName = RedirectedTopLevelAssetName.ObjectName.ToString();
-		Delta += InThumbnailEntry.ObjectShortClassName.Len() + InThumbnailEntry.ObjectPathWithoutPackageName.Len();
-
-		InThumbnailEntry.Delta = Delta;
 	}
 
 	return bPatched;
@@ -3010,13 +3002,11 @@ FAssetHeaderPatcher::EResult FAssetHeaderPatcherInner::PatchHeader_WriteDestinat
 			const int64 ThumbnailTableDeltaOffset = ThumbnailTableStartOffset - SourceSection.Offset;
 			int32 ThumbnailCount = ThumbnailTable.Num();
 			Writer << ThumbnailCount;
-			int32 AccumulatedDelta = 0;
 			for (FThumbnailEntry& Entry : ThumbnailTable)
 			{
-				AccumulatedDelta += Entry.Delta;
 				Writer << Entry.ObjectShortClassName;
 				Writer << Entry.ObjectPathWithoutPackageName;
-				Entry.FileOffset += (int32)ThumbnailTableDeltaOffset + AccumulatedDelta;
+				Entry.FileOffset += (int32)ThumbnailTableDeltaOffset;
 				Writer << Entry.FileOffset;
 			}
 			checkf(!Writer.IsCriticalError(), TEXT("Issue writing %s"), *Writer.GetErrorMessage());
