[AssetHeaderPatcher]

- Fix issue where Imports whose object name has been redirected while their package name has been redirected would not be patched correctly. We would incorrectly use the original ImportTable.ObjectName entry instead of the potentially new redirected name. e.g (/Package/_Verse/SomePath/Default__some_verse_class redirected to /Package/_Verse/Default__SomePath-some_verse_class has both package and object name redirects. Previously we would re-use the original ObjectName if sub-objects to this import where in the ImportTable instead of using the new object name for their outer leaving invalid entries in the ImportTable)
- Fixed issue where we would not consider any import whose package name had been redirected as patched and would not bother checking if its outer needed redirecting. Now we specifically check if an object has an outer chain, and if so always have its outer chain checked for redirection
- Adds test case for broken import patching and cleans up some test code.

#rnx

diff --git a/Engine/Source/Developer/AssetTools/Private/AssetHeaderPatcher.cpp b/Engine/Source/Developer/AssetTools/Private/AssetHeaderPatcher.cpp
--- a/Engine/Source/Developer/AssetTools/Private/AssetHeaderPatcher.cpp
+++ b/Engine/Source/Developer/AssetTools/Private/AssetHeaderPatcher.cpp
@@ -29,6 +29,8 @@
 
 DEFINE_LOG_CATEGORY_STATIC(LogAssetHeaderPatcher, Log, All);
 
+#define DEBUG_ASSET_HEADER_PATCHING 0
+
 namespace
 {
 	// If working on header patching, this is very helpful for dumping what is patched and reviewing the files in a folder comparison of your favourite diff program.
@@ -1085,7 +1087,11 @@ UE::Tasks::FTask FAssetHeaderPatcher::PatchAsync(int32* InOutNumFilesToPatch, in
 	// Note we are scheduling and launching tasks one at a time rather than preparing all jobs and launching all at once.
 	// While this means more overhead scheduling, it means that we won't have many tasks all hit the filesystem at the same time
 	// attempting to read and (more importantly) write to disk at the exact same time.
-	constexpr bool bSingleThreaded = false; // Useful for debugging
+#if DEBUG_ASSET_HEADER_PATCHING
+	constexpr bool bSingleThreaded = true; // Useful for debugging
+#else
+	constexpr bool bSingleThreaded = false;
+#endif
 	for (const TTuple<FString, FString>& Filename : PatchedFiles)
 	{
 		auto DoPatchFn = [this, 
@@ -1913,7 +1919,7 @@ FAssetHeaderPatcher::EResult FAssetHeaderPatcherInner::GetImportTablePatches(TAr
 		FCoreRedirectObjectName SrcImportPath;
 		FCoreRedirectObjectName DstImportPath;
 		bool bPatched = false;
-		bool bRedirectEvaluated = false;
+		bool bSkipImportTableWalkForRedirectedOuters = false;
 	};
 	TArray<FPatchDataForImport> ImportIndexToPatchData;
 	ImportIndexToPatchData.SetNum(ImportTable.Num());
@@ -1942,22 +1948,23 @@ FAssetHeaderPatcher::EResult FAssetHeaderPatcherInner::GetImportTablePatches(TAr
 
 		// Look up whether there is a specific redirect for the full path of the import.
 		PatchDataForImport.SrcImportPath = GetFullObjectNameFromObjectResource(Import, false /*bIsExport*/, false /*bWalkImportsOnly*/);
-		PatchDataForImport.DstImportPath = FCoreRedirects::GetRedirectedName(ECoreRedirectFlags::Type_AllMask, PatchDataForImport.SrcImportPath);
-		// If there is not a specific redirect, the object might still be redirected because its outer is redirected. All objects within
-		// a redirected outer are moved along with that outer, but the CoreRedirects system does not find the redirect of the outer unless
-		// we look for it specifically. We handle this case in the loop below, after we have processed all imports looking for redirects
-		// specific to them.
+		PatchDataForImport.DstImportPath = FCoreRedirects::GetRedirectedName(FCoreRedirects::GetFlagsForTypeName(Import.ClassPackage, Import.ClassName), PatchDataForImport.SrcImportPath);
 
 		// Record the remapping of the FName in the FObjectResource (whether or not it was redirected)
 		FName SrcObjectResourceName = GetObjectResourceNameFromCoreRedirectObjectName(PatchDataForImport.SrcImportPath);
 		FName DstObjectResourceName = GetObjectResourceNameFromCoreRedirectObjectName(PatchDataForImport.DstImportPath);
 		RemapFName(SrcObjectResourceName, DstObjectResourceName);
 
-		// If there was a redirect specifically for this import, create an FImportPatch
+		// All objects within a redirected outer are moved along with that outer, but the CoreRedirects system does not find the redirect of the outer unless
+		// we look for it specifically. We handle this case in the loop below, after we have processed all imports looking for redirects specific to them.
 		PatchDataForImport.bPatched = PatchDataForImport.SrcImportPath != PatchDataForImport.DstImportPath;
 		if (PatchDataForImport.bPatched)
 		{
-			PatchDataForImport.bRedirectEvaluated = true;
+			// If our outer was patched (i.e. we had a direct match) or there is no outer but our package changed, ensure we don't use the ImportTable to 
+			// check for redirected outers that should apply to this Import as we may incorrectly change an OuterName/PackageName back to the SrcName 
+			// if that outer name is still in use after redirection. (e.g. an object moved out of a package to a new package but both packages are still valid imports)
+			PatchDataForImport.bSkipImportTableWalkForRedirectedOuters = (PatchDataForImport.SrcImportPath.OuterName != PatchDataForImport.DstImportPath.OuterName) ||
+				(PatchDataForImport.DstImportPath.OuterName.IsNone() && PatchDataForImport.SrcImportPath.PackageName != PatchDataForImport.DstImportPath.PackageName);
 			PatchDataForImport.PatchIndex = OutImportPatches.Num();
 
 			FImportPatch& ImportPatch = OutImportPatches.Emplace_GetRef();
@@ -1981,7 +1988,7 @@ FAssetHeaderPatcher::EResult FAssetHeaderPatcherInner::GetImportTablePatches(TAr
 		FObjectImport& Import = ImportTable[ImportIndex];
 		FPatchDataForImport& PatchDataForImport = ImportIndexToPatchData[ImportIndex];
 
-		if (!PatchDataForImport.bRedirectEvaluated)
+		if (!PatchDataForImport.bSkipImportTableWalkForRedirectedOuters)
 		{
 			// Recursively evaluate all outers of the import, and then evaluate the import
 			TArray<int32, TInlineAllocator<10>> ImportsToEvaluate;
@@ -1999,7 +2006,7 @@ FAssetHeaderPatcher::EResult FAssetHeaderPatcherInner::GetImportTablePatches(TAr
 				int32 OuterImportIndex = OuterIndex.ToImport();
 				FObjectImport& OuterImport = ImportTable[OuterImportIndex];
 				FPatchDataForImport& PatchDataForOuter = ImportIndexToPatchData[OuterImportIndex];
-				if (PatchDataForOuter.bRedirectEvaluated)
+				if (PatchDataForOuter.bSkipImportTableWalkForRedirectedOuters)
 				{
 					break;
 				}
@@ -2012,7 +2019,7 @@ FAssetHeaderPatcher::EResult FAssetHeaderPatcherInner::GetImportTablePatches(TAr
 				int32 ImportIndexToEvaluate = ImportsToEvaluate.Pop(EAllowShrinking::No);
 				FObjectImport& ImportToEvaluate = ImportTable[ImportIndexToEvaluate];
 				FPatchDataForImport& PatchDataToEvaluate = ImportIndexToPatchData[ImportIndexToEvaluate];
-				if (PatchDataToEvaluate.bRedirectEvaluated)
+				if (PatchDataToEvaluate.bSkipImportTableWalkForRedirectedOuters)
 				{
 					// Shouldn't be possible, but somehow we already evaluated this, cycle in the outer chain?
 					continue;
@@ -2021,29 +2028,28 @@ FAssetHeaderPatcher::EResult FAssetHeaderPatcherInner::GetImportTablePatches(TAr
 				if (ImportToEvaluate.OuterIndex.IsNull())
 				{
 					// This import has no outer, it is a package. We did not find a redirect specifically for it,
-					// and it has no outer to move along with, so it has no redirect. Mark it evaluated and do nothing.
-					PatchDataToEvaluate.bRedirectEvaluated = true;
+					// and it has no outer to move along with, so it has no redirect. Mark it skipped and do nothing.
+					PatchDataToEvaluate.bSkipImportTableWalkForRedirectedOuters = true;
 				}
 				else if (ImportToEvaluate.OuterIndex.IsExport())
 				{
 					// We don't have data on the export, so we cannot read the transformed path for the import under it.
 					// TODO: This is incorrect, we need to get the remapped destination from the export to know
 					// the correct destination for the import.
-					PatchDataToEvaluate.bRedirectEvaluated = true;
+					PatchDataToEvaluate.bSkipImportTableWalkForRedirectedOuters = true;
 				}
 				else
 				{
 					int32 OuterImportIndex = ImportToEvaluate.OuterIndex.ToImport();
-					FObjectImport& OuterImport = ImportTable[OuterImportIndex];
 					FPatchDataForImport& PatchDataForOuter = ImportIndexToPatchData[OuterImportIndex];
 					PatchDataToEvaluate.DstImportPath = FCoreRedirectObjectName::AppendObjectName(
-						PatchDataForOuter.DstImportPath, ImportToEvaluate.ObjectName);
-					PatchDataToEvaluate.bRedirectEvaluated = true;
+						PatchDataForOuter.DstImportPath, PatchDataToEvaluate.DstImportPath.ObjectName);
+					PatchDataToEvaluate.bSkipImportTableWalkForRedirectedOuters = true;
 				}
 			}
 		}
 
-		check(PatchDataForImport.bRedirectEvaluated);
+		check(PatchDataForImport.bSkipImportTableWalkForRedirectedOuters);
 
 		// Record that the (possibly redirected) full path of the destination of this import is in this import index
 		ImportNameToImportTableIndexLookup.Add(PatchDataForImport.DstImportPath, ImportIndex);
@@ -3612,7 +3618,7 @@ void FAssetHeaderPatcherInner::DumpState(FStringView OutputDirectory)
 TEST_CASE_NAMED(FAssetHeaderPatcherTests, "AssetHeaderPatcher", "[AssetHeaderPatcher][EngineFilter]")
 {
 // Useful when iterating so you halt if something fails while a debugger is attached
-#if 0
+#if DEBUG_ASSET_HEADER_PATCHING
 #undef CHECK
 #undef CHECK_EQUALS
 #undef CHECK_NOT_EQUALS
@@ -3665,24 +3671,32 @@ TEST_CASE_NAMED(FAssetHeaderPatcherTests, "AssetHeaderPatcher", "[AssetHeaderPat
 	const FName DummyImportPackagePathFName(TEXT("/DummyMount/DummyPackage"));
 	// Import FNames
 	const FName SrcEngineModuleImportObjectName(TEXT("/Script/SrcEngineModule"));
-	const FName Default__SrcTypeAImportObjectName(TEXT("Default__SrcTypeA"));
-	const FName Default__SrcTypeBImportObjectName(TEXT("Default__SrcTypeB"));
-	const FName Default__OnlySubTypeChangedImportObjectName(TEXT("Default__OnlySubTypeChanged"));
-	const FName Default__MovedButNotRenamedTypeImportObjectName(TEXT("Default__MovedButNotRenamedType"));
+	const FName SrcTypeAImportObjectName(TEXT("SrcTypeA"));
+	const FName SrcTypeBImportObjectName(TEXT("SrcTypeB"));
+	const FName OnlySubTypeChangedImportObjectName(TEXT("OnlySubTypeChanged"));
+	const FName MovedButNotRenamedTypeImportObjectName(TEXT("MovedButNotRenamedType"));
 	const FName SrcPropertyToChangeAImportObjectName(TEXT("SrcPropertyToChangeA"));
 	const FName SrcPropertyToChangeBImportObjectName(TEXT("SrcPropertyToChangeB"));
 	const FName MovedButNotRenamedPropertyImportObjectName(TEXT("MovedButNotRenamedProperty"));
+	const FName NewOuterImportObjectName(TEXT("NewOuter"));
+	const FName MovedToNewOuterImportObjectName(TEXT("MovedToNewOuter"));
 	const FName InnerMovedButNotRenamedPropertyImportObjectName(TEXT("InnerMovedButNotRenamedProperty"));
 	const FName InnerInnerMovedButNotRenamedPropertyImportObjectName(TEXT("InnerInnerMovedButNotRenamedProperty"));
 	const FName UnchangedPropertyImportObjectName(TEXT("UnchangedProperty"));
 	const FName SrcImportClassName(TEXT("SrcClass"));
 	const FName SrcImportClassPackage(TEXT("/Engine/SrcClassPackage"));
-	
+	const FName SrcVerseAssetName(TEXT("/Module/_Verse/VerseAsset"));
+	const FName SrcVerseImportObjectName(TEXT("some_verse_class"));
+	const FName UnchangedVerseImportSubObject1Name(TEXT("__verse_0x7A8CDEBC_VerseObject1"));
+	const FName UnchangedVerseImportSubObject2Name(TEXT("__verse_0x5614AC82_VerseObject2"));
+
 	const FName DstEngineModuleImportObjectName(TEXT("/Script/DstEngineModule"));
-	const FName Default__DstTypeAImportObjectName(TEXT("Default__DstTypeA"));
-	const FName Default__DstTypeBImportObjectName(TEXT("Default__DstTypeB"));
+	const FName DstTypeAImportObjectName(TEXT("DstTypeA"));
+	const FName DstTypeBImportObjectName(TEXT("DstTypeB"));
 	const FName DstPropertyToChangeAImportObjectName(TEXT("DstPropertyToChangeA"));
 	const FName DstPropertyToChangeBImportObjectName(TEXT("DstPropertyToChangeB"));
+	const FName DstVerseAssetName(TEXT("/Module/_Verse"));
+	const FName DstVerseImportObjectName(TEXT("VerseAsset-some_verse_class"));
 	const FName DstImportClassName(TEXT("DstClass"));
 	const FName DstImportClassPackage(TEXT("/Engine/DstClassPackage"));
 
@@ -3733,96 +3747,141 @@ TEST_CASE_NAMED(FAssetHeaderPatcherTests, "AssetHeaderPatcher", "[AssetHeaderPat
 		FObjectImport Dst;
 		bool bExistingImport;
 	};
+	// Note the order of these cases defines the ImportTable entry order before/after patching
 	TArray<FImportTestCase> ImportTestCases
 	{
-		// /Script/SrcEngineModule																-> (remains unchanged)
+		// /Script/SrcEngineModule -> (remains unchanged)
 		{	
-			.Src = MakeImport(SrcEngineModuleImportObjectName, FPackageIndex(), SrcImportClassPackage, SrcImportClassName, SrcPackagePathFName), 
-			.Dst = MakeImport(SrcEngineModuleImportObjectName, FPackageIndex(), DstImportClassPackage, DstImportClassName, DstPackagePathFName),
+			.Src = MakeImport(SrcEngineModuleImportObjectName, FPackageIndex(), GLongCoreUObjectPackageName, NAME_Package, SrcPackagePathFName),
+			.Dst = MakeImport(SrcEngineModuleImportObjectName, FPackageIndex(), GLongCoreUObjectPackageName, NAME_Package, DstPackagePathFName),
 			.bExistingImport = true,
 		},
-		// /Script/SrcEngineModule.Default__SrcTypeA											-> /Script/DstEngineModule.Default__DstTypeA
+		// /Script/SrcEngineModule.SrcTypeA	-> /Script/DstEngineModule.DstTypeA
 		{ 
-			.Src = MakeImport(Default__SrcTypeAImportObjectName, FPackageIndex::FromImport(0), SrcImportClassPackage, SrcImportClassName, SrcPackagePathFName), 
-			.Dst = MakeImport(Default__DstTypeAImportObjectName, FPackageIndex::FromImport(13), DstImportClassPackage, DstImportClassName, DstPackagePathFName),
+			.Src = MakeImport(SrcTypeAImportObjectName, FPackageIndex::FromImport(0), SrcImportClassPackage, SrcImportClassName, SrcPackagePathFName), 
+			.Dst = MakeImport(DstTypeAImportObjectName, FPackageIndex::FromImport(18), DstImportClassPackage, DstImportClassName, DstPackagePathFName),
 			.bExistingImport = true,
 		},
-		// /Script/SrcEngineModule.Default__SrcTypeA.SrcPropertyToChangeA						-> /Script/DstEngineModule.Default__DstTypeA.DstPropertyToChangeA
+		// /Script/SrcEngineModule.SrcTypeA.SrcPropertyToChangeA -> /Script/DstEngineModule.DstTypeA.DstPropertyToChangeA
 		{ 
 			.Src = MakeImport(SrcPropertyToChangeAImportObjectName, FPackageIndex::FromImport(1), SrcImportClassPackage, SrcImportClassName, SrcPackagePathFName), 
 			.Dst = MakeImport(DstPropertyToChangeAImportObjectName, FPackageIndex::FromImport(1), DstImportClassPackage, DstImportClassName, DstPackagePathFName),
 			.bExistingImport = true,
 		},
-		// /Script/SrcEngineModule.Default__SrcTypeA.MovedButNotRenamedProperty					-> /Script/DstEngineModule.Default__DstTypeA.MovedButNotRenamedProperty
+		// /Script/SrcEngineModule.SrcTypeA.MovedButNotRenamedProperty -> /Script/DstEngineModule.DstTypeA.MovedButNotRenamedProperty
 		{ 
 			.Src = MakeImport(MovedButNotRenamedPropertyImportObjectName, FPackageIndex::FromImport(1), SrcImportClassPackage, SrcImportClassName, SrcPackagePathFName), 
 			.Dst = MakeImport(MovedButNotRenamedPropertyImportObjectName, FPackageIndex::FromImport(1), DstImportClassPackage, DstImportClassName, DstPackagePathFName), 
 			.bExistingImport = true,
 		},
-		// /Script/SrcEngineModule.Default__SrcTypeA.MovedButNotRenamedProperty.InnerMovedButNotRenamedProperty
+		// /Script/SrcEngineModule.SrcTypeA.MovedButNotRenamedProperty.InnerMovedButNotRenamedProperty
 		// -> 
-		// /Script/DstEngineModule.Default__DstTypeA.MovedButNotRenamedProperty.InnerMovedButNotRenamedProperty
+		// /Script/DstEngineModule.DstTypeA.MovedButNotRenamedProperty.InnerMovedButNotRenamedProperty
 		{
 			.Src = MakeImport(InnerMovedButNotRenamedPropertyImportObjectName, FPackageIndex::FromImport(3), SrcImportClassPackage, SrcImportClassName, SrcPackagePathFName),
 			.Dst = MakeImport(InnerMovedButNotRenamedPropertyImportObjectName, FPackageIndex::FromImport(3), DstImportClassPackage, DstImportClassName, DstPackagePathFName),
 			.bExistingImport = true,
 		},
-		// /Script/SrcEngineModule.Default__SrcTypeA.MovedButNotRenamedProperty.InnerMovedButNotRenamedProperty.InnerInnerMovedButNotRenamedProperty 
+		// /Script/SrcEngineModule.SrcTypeA.MovedButNotRenamedProperty.InnerMovedButNotRenamedProperty.InnerInnerMovedButNotRenamedProperty 
 		// -> 
-		// /Script/DstEngineModule.Default__DstTypeA.MovedButNotRenamedProperty.InnerMovedButNotRenamedProperty.InnerInnerMovedButNotRenamedProperty 
+		// /Script/DstEngineModule.DstTypeA.MovedButNotRenamedProperty.InnerMovedButNotRenamedProperty.InnerInnerMovedButNotRenamedProperty 
 		{
 			.Src = MakeImport(InnerInnerMovedButNotRenamedPropertyImportObjectName, FPackageIndex::FromImport(4), SrcImportClassPackage, SrcImportClassName, SrcPackagePathFName),
 			.Dst = MakeImport(InnerInnerMovedButNotRenamedPropertyImportObjectName, FPackageIndex::FromImport(4), DstImportClassPackage, DstImportClassName, DstPackagePathFName),
 			.bExistingImport = true,
 		},
-		// /Script/SrcEngineModule.Default__OnlySubTypeChanged									-> (remains unchanged)
+		// /Script/SrcEngineModule.OnlySubTypeChanged -> (remains unchanged)
 		{ 
-			.Src = MakeImport(Default__OnlySubTypeChangedImportObjectName, FPackageIndex::FromImport(0), SrcImportClassPackage, SrcImportClassName, SrcPackagePathFName), 
-			.Dst = MakeImport(Default__OnlySubTypeChangedImportObjectName, FPackageIndex::FromImport(0), DstImportClassPackage, DstImportClassName, DstPackagePathFName),
+			.Src = MakeImport(OnlySubTypeChangedImportObjectName, FPackageIndex::FromImport(0), SrcImportClassPackage, SrcImportClassName, SrcPackagePathFName), 
+			.Dst = MakeImport(OnlySubTypeChangedImportObjectName, FPackageIndex::FromImport(0), DstImportClassPackage, DstImportClassName, DstPackagePathFName),
 			.bExistingImport = true,
 		},
-		// /Script/SrcEngineModule.Default__OnlySubTypeChanged.SrcPropertyToChangeB				-> /Script/SrcEngineModule.Default__OnlySubTypeChanged.DstPropertyToChangeB
+		// /Script/SrcEngineModule.OnlySubTypeChanged.SrcPropertyToChangeB -> /Script/SrcEngineModule.OnlySubTypeChanged.DstPropertyToChangeB
 		{ 
 			.Src = MakeImport(SrcPropertyToChangeBImportObjectName, FPackageIndex::FromImport(6), SrcImportClassPackage, SrcImportClassName, SrcPackagePathFName), 
 			.Dst = MakeImport(DstPropertyToChangeBImportObjectName, FPackageIndex::FromImport(6), DstImportClassPackage, DstImportClassName, DstPackagePathFName),
 			.bExistingImport = true,
 		},
-		// /Script/SrcEngineModule.Default__OnlySubTypeChanged.UnchangedProperty				-> (remains unchanged)
+		// /Script/SrcEngineModule.OnlySubTypeChanged.UnchangedProperty -> (remains unchanged)
 		{ 
 			.Src = MakeImport(UnchangedPropertyImportObjectName, FPackageIndex::FromImport(6), SrcImportClassPackage, SrcImportClassName, SrcPackagePathFName), 
 			.Dst = MakeImport(UnchangedPropertyImportObjectName, FPackageIndex::FromImport(6), DstImportClassPackage, DstImportClassName, DstPackagePathFName),
 			.bExistingImport = true,
 		},
-		// /Script/SrcEngineModule.Default__SrcTypeB											-> /Script/SrcEngineModule.Default__DstTypeB
+		// /Script/SrcEngineModule.SrcTypeB	-> /Script/SrcEngineModule.DstTypeB
 		{ 
-			.Src = MakeImport(Default__SrcTypeBImportObjectName, FPackageIndex::FromImport(0), SrcImportClassPackage, SrcImportClassName, SrcPackagePathFName), 
-			.Dst = MakeImport(Default__DstTypeBImportObjectName, FPackageIndex::FromImport(0), DstImportClassPackage, DstImportClassName, DstPackagePathFName),
+			.Src = MakeImport(SrcTypeBImportObjectName, FPackageIndex::FromImport(0), SrcImportClassPackage, SrcImportClassName, SrcPackagePathFName), 
+			.Dst = MakeImport(DstTypeBImportObjectName, FPackageIndex::FromImport(0), DstImportClassPackage, DstImportClassName, DstPackagePathFName),
 			.bExistingImport = true,
 		},
-		// /Script/SrcEngineModule.Default__SrcTypeB.MovedButNotRenamedProperty					-> /Script/SrcEngineModule.Default__DstTypeB.MovedButNotRenamedProperty
+		// /Script/SrcEngineModule.SrcTypeB.MovedButNotRenamedProperty -> /Script/SrcEngineModule.DstTypeB.MovedButNotRenamedProperty
 		{ 
 			.Src = MakeImport(MovedButNotRenamedPropertyImportObjectName, FPackageIndex::FromImport(9), SrcImportClassPackage, SrcImportClassName, SrcPackagePathFName), 
 			.Dst = MakeImport(MovedButNotRenamedPropertyImportObjectName, FPackageIndex::FromImport(9), DstImportClassPackage, DstImportClassName, DstPackagePathFName),
 			.bExistingImport = true,
 		},
-		// /Script/SrcEngineModule.Default__MovedButNotRenamedType								-> /Script/DstEngineModule.Default__MovedButNotRenamedType
+		// /Script/SrcEngineModule.SrcTypeB.MovedToNewOuter	-> /Script/SrcEngineModule.NewOuter.MovedToNewOuter
+		{
+			.Src = MakeImport(MovedToNewOuterImportObjectName, FPackageIndex::FromImport(9), SrcImportClassPackage, SrcImportClassName, SrcPackagePathFName),
+			.Dst = MakeImport(MovedToNewOuterImportObjectName, FPackageIndex::FromImport(19), DstImportClassPackage, DstImportClassName, DstPackagePathFName),
+			.bExistingImport = true,
+		},
+		// /Script/SrcEngineModule.MovedButNotRenamedType -> /Script/DstEngineModule.MovedButNotRenamedType
 		{ 
-			.Src = MakeImport(Default__MovedButNotRenamedTypeImportObjectName, FPackageIndex::FromImport(0), SrcImportClassPackage, SrcImportClassName, SrcPackagePathFName), 
-			.Dst = MakeImport(Default__MovedButNotRenamedTypeImportObjectName, FPackageIndex::FromImport(13), DstImportClassPackage, DstImportClassName, DstPackagePathFName),
+			.Src = MakeImport(MovedButNotRenamedTypeImportObjectName, FPackageIndex::FromImport(0), SrcImportClassPackage, SrcImportClassName, SrcPackagePathFName), 
+			.Dst = MakeImport(MovedButNotRenamedTypeImportObjectName, FPackageIndex::FromImport(18), DstImportClassPackage, DstImportClassName, DstPackagePathFName),
 			.bExistingImport = true,
 		},
-		// /Script/SrcEngineModule.Default__MovedButNotRenamedType.MovedButNotRenamedProperty	-> /Script/DstEngineModule.Default__MovedButNotRenamedType.MovedButNotRenamedProperty
+		// /Script/SrcEngineModule.MovedButNotRenamedType.MovedButNotRenamedProperty -> /Script/DstEngineModule.MovedButNotRenamedType.MovedButNotRenamedProperty
 		{ 
-			.Src = MakeImport(MovedButNotRenamedPropertyImportObjectName, FPackageIndex::FromImport(11), SrcImportClassPackage, SrcImportClassName, SrcPackagePathFName), 
-			.Dst = MakeImport(MovedButNotRenamedPropertyImportObjectName, FPackageIndex::FromImport(11), DstImportClassPackage, DstImportClassName, DstPackagePathFName),
+			.Src = MakeImport(MovedButNotRenamedPropertyImportObjectName, FPackageIndex::FromImport(12), SrcImportClassPackage, SrcImportClassName, SrcPackagePathFName), 
+			.Dst = MakeImport(MovedButNotRenamedPropertyImportObjectName, FPackageIndex::FromImport(12), DstImportClassPackage, DstImportClassName, DstPackagePathFName),
+			.bExistingImport = true,
+		},
+		// Note, the verse transformations test a case where package is moved into a new package with a different path, while at the same time renaming top-level assets 
+		// in the new path. We want to ensure we don't keep using the old top-level asset name. These redirections are done using CoreRedirects from the context
+		// rather than explicit remappings in the patcher.
+		
+		// /Module/_Verse/VerseAsset -> /Module/_Verse
+		{
+			.Src = MakeImport(SrcVerseAssetName, FPackageIndex(), GLongCoreUObjectPackageName, NAME_Package, SrcPackagePathFName),
+			.Dst = MakeImport(DstVerseAssetName, FPackageIndex(), GLongCoreUObjectPackageName, NAME_Package, DstPackagePathFName),
+			.bExistingImport = true,
+		},
+		// /Module/_Verse/VerseAsset.some_verse_class -> /Module/_Verse.VerseAsset-some_verse_class
+		{
+			.Src = MakeImport(SrcVerseImportObjectName, FPackageIndex::FromImport(14), SrcImportClassPackage, SrcImportClassName, SrcPackagePathFName),
+			.Dst = MakeImport(DstVerseImportObjectName, FPackageIndex::FromImport(14), DstImportClassPackage, DstImportClassName, DstPackagePathFName),
+			.bExistingImport = true,
+		},
+		// /Module/_Verse/VerseAsset.some_verse_class.__verse_0x7A8CDEBC_VerseObject1 
+		// -> 
+		// /Module/_Verse.VerseAsset-some_verse_class.__verse_0x7A8CDEBC_VerseObject1
+		{
+			.Src = MakeImport(UnchangedVerseImportSubObject1Name, FPackageIndex::FromImport(15), SrcImportClassPackage, SrcImportClassName, SrcPackagePathFName),
+			.Dst = MakeImport(UnchangedVerseImportSubObject1Name, FPackageIndex::FromImport(15), DstImportClassPackage, DstImportClassName, DstPackagePathFName),
+			.bExistingImport = true,
+		},
+		// /Module/_Verse/VerseAsset.some_verse_class.__verse_0x7A8CDEBC_VerseObject1.__verse_0x5614AC82_VerseObject2
+		// ->
+		// /Module/_Verse.VerseAsset-some_verse_class.__verse_0x7A8CDEBC_VerseObject1.__verse_0x5614AC82_VerseObject2
+		{
+			.Src = MakeImport(UnchangedVerseImportSubObject2Name, FPackageIndex::FromImport(16), SrcImportClassPackage, SrcImportClassName, SrcPackagePathFName),
+			.Dst = MakeImport(UnchangedVerseImportSubObject2Name, FPackageIndex::FromImport(16), DstImportClassPackage, DstImportClassName, DstPackagePathFName),
 			.bExistingImport = true,
 		},
-		// <none>																				-> /Script/DstEngineModule
+		// <none> -> /Script/DstEngineModule
 		{ 
 			.Src = FObjectImport(), 
 			.Dst = MakeImport(DstEngineModuleImportObjectName, FPackageIndex(), DstImportClassPackage, DstImportClassName, DstPackagePathFName),
 			.bExistingImport = false,
 		},
+		// <none> -> /Script/SrcEngineModule.NewOuter
+		{
+			.Src = FObjectImport(),
+			.Dst = MakeImport(NewOuterImportObjectName, FPackageIndex::FromImport(0), DstImportClassPackage, DstImportClassName, DstPackagePathFName),
+			.bExistingImport = false,
+		},
 	};
 
 	struct FExportTestCase
@@ -3832,18 +3891,18 @@ TEST_CASE_NAMED(FAssetHeaderPatcherTests, "AssetHeaderPatcher", "[AssetHeaderPat
 	};
 	TArray<FExportTestCase> ExportTestCases
 	{
-		// SrcPackage																			-> DstPackage
-		// Super: /Script/SrcEngineModule														-> /Script/DstEngineModule
-		// Class: /Script/SrcEngineModule.Default__SrcTypeA.SrcPropertyToChangeA				-> /Script/DstEngineModule.Default__DstTypeA.DstPropertyToChangeA
-		// Template: /Script/SrcEngineModule.Default__OnlySubTypeChanged						-> /Script/SrcEngineModule.Default__OnlySubTypeChanged
+		// SrcPackage																-> DstPackage
+		// Super: /Script/SrcEngineModule											-> /Script/DstEngineModule
+		// Class: /Script/SrcEngineModule.SrcTypeA.SrcPropertyToChangeA				-> /Script/DstEngineModule.DstTypeA.DstPropertyToChangeA
+		// Template: /Script/SrcEngineModule.OnlySubTypeChanged						-> /Script/SrcEngineModule.OnlySubTypeChanged
 		{
 			.Src = MakeExport(SrcExportObjectFName, FPackageIndex::FromExport(0), FPackageIndex(), FPackageIndex::FromImport(0), FPackageIndex::FromImport(2), FPackageIndex::FromImport(4)),
 			.Dst = MakeExport(DstExportObjectFName, FPackageIndex::FromExport(0), FPackageIndex(), FPackageIndex::FromImport(11), FPackageIndex::FromImport(2), FPackageIndex::FromImport(4))
 		},
-		// SrcPackage.SrcPackage																-> DstPackage.DstPackage
-		// Super: /Script/SrcEngineModule.Default__SrcTypeA										-> /Script/DstEngineModule.Default__DstTypeA
-		// Class: /Script/SrcEngineModule.Default__SrcTypeA.MovedButNotRenamedProperty			-> /Script/DstEngineModule.Default__DstTypeA.MovedButNotRenamedProperty
-		// Template: /Script/SrcEngineModule.Default__OnlySubTypeChanged.SrcPropertyToChangeB	-> /Script/SrcEngineModule.Default__OnlySubTypeChanged.DstPropertyToChangeB
+		// SrcPackage.SrcPackage														-> DstPackage.DstPackage
+		// Super: /Script/SrcEngineModule.SrcTypeA										-> /Script/DstEngineModule.DstTypeA
+		// Class: /Script/SrcEngineModule.SrcTypeA.MovedButNotRenamedProperty			-> /Script/DstEngineModule.DstTypeA.MovedButNotRenamedProperty
+		// Template: /Script/SrcEngineModule.OnlySubTypeChanged.SrcPropertyToChangeB	-> /Script/SrcEngineModule.OnlySubTypeChanged.DstPropertyToChangeB
 		{
 			.Src = MakeExport(SrcExportObjectFName, FPackageIndex::FromExport(1), FPackageIndex::FromExport(0), FPackageIndex::FromImport(1), FPackageIndex::FromImport(3), FPackageIndex::FromImport(5)),
 			.Dst = MakeExport(DstExportObjectFName, FPackageIndex::FromExport(1), FPackageIndex::FromExport(0), FPackageIndex::FromImport(1), FPackageIndex::FromImport(3), FPackageIndex::FromImport(5))
@@ -3864,6 +3923,11 @@ TEST_CASE_NAMED(FAssetHeaderPatcherTests, "AssetHeaderPatcher", "[AssetHeaderPat
 
 	FAssetHeaderPatcherInner Patcher(DummySrcDstAsset, DummySrcDstAsset, StringReplacements, MountPointReplacementMap);
 
+	auto AddToNameTable = [&Patcher](FName Name)
+		{
+			Patcher.NameToIndexMap.Add(Name.GetDisplayIndex(), Patcher.NameTable.Num());
+			Patcher.NameTable.Add(Name);
+		};
 	int32 OriginalNameTableCount = 0;
 	auto ResetPatcher = [&]()
 		{
@@ -3874,58 +3938,50 @@ TEST_CASE_NAMED(FAssetHeaderPatcherTests, "AssetHeaderPatcher", "[AssetHeaderPat
 			///////////////////////////////////////////////////////////////////////////////////
 			
 			// NameTable
-			Patcher.NameToIndexMap.Add(SrcPackagePathFName.GetDisplayIndex(), Patcher.NameTable.Num());
-			Patcher.NameTable.Add(SrcPackagePathFName);
-			Patcher.NameToIndexMap.Add(SrcAssetFName.GetDisplayIndex(), Patcher.NameTable.Num());
-			Patcher.NameTable.Add(SrcAssetFName);
-			Patcher.NameToIndexMap.Add(FName(NAME_None).GetDisplayIndex(), Patcher.NameTable.Num());
-			Patcher.NameTable.Add(FName(NAME_None));
-			Patcher.NameToIndexMap.Add(FName(DummyImportPackagePathFName).GetDisplayIndex(), Patcher.NameTable.Num());
-			Patcher.NameTable.Add(FName(DummyImportPackagePathFName));
+			AddToNameTable(SrcPackagePathFName);
+			AddToNameTable(SrcAssetFName);
+			AddToNameTable(DummyImportPackagePathFName);
 			// Redirected names (softpaths don't include the subpath in the name table
-			FName ToBeRedirectedPackageName = SoftObjectPathToRedirect.GetAssetPath().GetPackageName();
-			FName ToBeRedirectedAssetName = SoftObjectPathToRedirect.GetAssetPath().GetAssetName();
-			Patcher.NameToIndexMap.Add(ToBeRedirectedPackageName.GetDisplayIndex(), Patcher.NameTable.Num());
-			Patcher.NameTable.Add(ToBeRedirectedPackageName);
-			Patcher.NameToIndexMap.Add(ToBeRedirectedAssetName.GetDisplayIndex(), Patcher.NameTable.Num());
-			Patcher.NameTable.Add(ToBeRedirectedAssetName);
+			AddToNameTable(SoftObjectPathToRedirect.GetAssetPath().GetPackageName());
+			AddToNameTable(SoftObjectPathToRedirect.GetAssetPath().GetAssetName());
 			// Export Table Names
-			Patcher.NameToIndexMap.Add(SrcExportObjectFName.GetDisplayIndex(), Patcher.NameTable.Num());
-			Patcher.NameTable.Add(SrcExportObjectFName);
+			AddToNameTable(SrcExportObjectFName);
 			// Import Table Names
-			Patcher.NameToIndexMap.Add(SrcEngineModuleImportObjectName.GetDisplayIndex(), Patcher.NameTable.Num());
-			Patcher.NameTable.Add(SrcEngineModuleImportObjectName);
-			Patcher.NameToIndexMap.Add(Default__SrcTypeAImportObjectName.GetDisplayIndex(), Patcher.NameTable.Num());
-			Patcher.NameTable.Add(Default__SrcTypeAImportObjectName);
-			Patcher.NameToIndexMap.Add(Default__SrcTypeBImportObjectName.GetDisplayIndex(), Patcher.NameTable.Num());
-			Patcher.NameTable.Add(Default__SrcTypeBImportObjectName);
-			Patcher.NameToIndexMap.Add(Default__OnlySubTypeChangedImportObjectName.GetDisplayIndex(), Patcher.NameTable.Num());
-			Patcher.NameTable.Add(Default__OnlySubTypeChangedImportObjectName);
-			Patcher.NameToIndexMap.Add(Default__MovedButNotRenamedTypeImportObjectName.GetDisplayIndex(), Patcher.NameTable.Num());
-			Patcher.NameTable.Add(Default__MovedButNotRenamedTypeImportObjectName);
-			Patcher.NameToIndexMap.Add(SrcPropertyToChangeAImportObjectName.GetDisplayIndex(), Patcher.NameTable.Num());
-			Patcher.NameTable.Add(SrcPropertyToChangeAImportObjectName);
-			Patcher.NameToIndexMap.Add(SrcPropertyToChangeBImportObjectName.GetDisplayIndex(), Patcher.NameTable.Num());
-			Patcher.NameTable.Add(SrcPropertyToChangeBImportObjectName);
-			Patcher.NameToIndexMap.Add(MovedButNotRenamedPropertyImportObjectName.GetDisplayIndex(), Patcher.NameTable.Num());
-			Patcher.NameTable.Add(MovedButNotRenamedPropertyImportObjectName);
-			Patcher.NameToIndexMap.Add(InnerMovedButNotRenamedPropertyImportObjectName.GetDisplayIndex(), Patcher.NameTable.Num());
-			Patcher.NameTable.Add(InnerMovedButNotRenamedPropertyImportObjectName);
-			Patcher.NameToIndexMap.Add(InnerInnerMovedButNotRenamedPropertyImportObjectName.GetDisplayIndex(), Patcher.NameTable.Num());
-			Patcher.NameTable.Add(InnerInnerMovedButNotRenamedPropertyImportObjectName);
-			Patcher.NameToIndexMap.Add(UnchangedPropertyImportObjectName.GetDisplayIndex(), Patcher.NameTable.Num());
-			Patcher.NameTable.Add(UnchangedPropertyImportObjectName);
-			Patcher.NameToIndexMap.Add(SrcImportClassName.GetDisplayIndex(), Patcher.NameTable.Num());
-			Patcher.NameTable.Add(SrcImportClassName);
-			Patcher.NameToIndexMap.Add(SrcImportClassPackage.GetDisplayIndex(), Patcher.NameTable.Num());
-			Patcher.NameTable.Add(SrcImportClassPackage);
+			AddToNameTable(SrcEngineModuleImportObjectName);
+			AddToNameTable(SrcTypeAImportObjectName);
+			AddToNameTable(SrcTypeBImportObjectName);
+			AddToNameTable(OnlySubTypeChangedImportObjectName);
+			AddToNameTable(MovedButNotRenamedTypeImportObjectName);
+			AddToNameTable(SrcPropertyToChangeAImportObjectName);
+			AddToNameTable(SrcPropertyToChangeBImportObjectName);
+			AddToNameTable(MovedButNotRenamedPropertyImportObjectName);
+			AddToNameTable(NewOuterImportObjectName);
+			AddToNameTable(MovedToNewOuterImportObjectName);
+			AddToNameTable(InnerMovedButNotRenamedPropertyImportObjectName);
+			AddToNameTable(InnerInnerMovedButNotRenamedPropertyImportObjectName);
+			AddToNameTable(UnchangedPropertyImportObjectName);
+			AddToNameTable(SrcVerseAssetName);
+			AddToNameTable(SrcVerseImportObjectName);
+			AddToNameTable(UnchangedVerseImportSubObject1Name);
+			AddToNameTable(UnchangedVerseImportSubObject2Name);
+			AddToNameTable(SrcImportClassName);
+			AddToNameTable(SrcImportClassPackage);
+			AddToNameTable(GLongCoreUObjectPackageName);
+			AddToNameTable(NAME_Package);
 
 			auto CheckNameTableInit = [&Patcher](FName Name)
 				{
+					// Skip lookups for None, as we already verified the NameTable 
+					// does not / should not contains None
+					if (Name == NAME_None)
+					{
+						return;
+					}
 					CHECK(Patcher.NameToIndexMap.Contains(Name.GetDisplayIndex()));
 					CHECK(Patcher.NameTable[Patcher.NameToIndexMap[Name.GetDisplayIndex()]] == Name);
 				};
-
+			CHECK(!Patcher.NameToIndexMap.Contains(FName(NAME_None).GetDisplayIndex()));
+			CHECK(!Patcher.NameTable.Contains(NAME_None));
 			// Import/Export Table (see breakdown of import names for the intended tests in the FObjectResource test section below)
 			for (const FImportTestCase& TestCase : ImportTestCases)
 			{
@@ -4431,27 +4487,40 @@ TEST_CASE_NAMED(FAssetHeaderPatcherTests, "AssetHeaderPatcher", "[AssetHeaderPat
 	// 
 	// Coreredirects are added to make the following transformations:
 	// 
-	// /Script/SrcEngineModule.Default__SrcTypeA.SrcPropertyToChangeA                     -> /Script/DstEngineModule.Default__DstTypeA.DstPropertyToChangeA
-	// /Script/SrcEngineModule.Default__SrcTypeA.MovedButNotRenamedProperty               -> /Script/DstEngineModule.Default__DstTypeA.MovedButNotRenamedProperty
-	// /Script/SrcEngineModule.Default__OnlySubTypeChanged.SrcPropertyToChangeB           -> /Script/SrcEngineModule.Default__OnlySubTypeChanged.DstPropertyToChangeB
-	// /Script/SrcEngineModule.Default__SrcTypeB.MovedButNotRenamedProperty               -> /Script/SrcEngineModule.Default__DstTypeB.MovedButNotRenamedProperty
-	// /Script/SrcEngineModule.Default__MovedButNotRenamedType.MovedButNotRenamedProperty -> /Script/DstEngineModule.Default__MovedButNotRenamedType.MovedButNotRenamedProperty
+	// /Engine/SrcClassPackage                                                   -> /Engine/DstClassPackage
+	// /Engine/SrcClassPackage.SrcClass                                          -> /Engine/DstClassPackage.DstClass
+	// /Script/SrcEngineModule.SrcTypeA                                          -> /Script/DstEngineModule.DstTypeA
+	// /Script/SrcEngineModule.SrcTypeA.SrcPropertyToChangeA                     -> /Script/DstEngineModule.DstTypeA.DstPropertyToChangeA
+	// /Script/SrcEngineModule.SrcTypeA.MovedButNotRenamedProperty               -> /Script/DstEngineModule.DstTypeA.MovedButNotRenamedProperty
+	// /Script/SrcEngineModule.OnlySubTypeChanged.SrcPropertyToChangeB           -> /Script/SrcEngineModule.OnlySubTypeChanged.DstPropertyToChangeB
+	// /Script/SrcEngineModule.SrcTypeB                                          -> /Script/SrcEngineModule.DstTypeB
+	// /Script/SrcEngineModule.SrcTypeB.MovedButNotRenamedProperty               -> /Script/SrcEngineModule.DstTypeB.MovedButNotRenamedProperty
+	// /Script/SrcEngineModule.SrcTypeB.MovedToNewOuter                          -> /Script/SrcEngineModule.NewOuter.MovedToNewOuter
+	// /Script/SrcEngineModule.MovedButNotRenamedType                            -> /Script/DstEngineModule.MovedButNotRenamedType
+	// /Script/SrcEngineModule.MovedButNotRenamedType.MovedButNotRenamedProperty -> /Script/DstEngineModule.MovedButNotRenamedType.MovedButNotRenamedProperty
+	// /Module/_Verse/VerseAsset                                                 -> /Module/_Verse
+	// /Module/_Verse/VerseAsset.some_verse_class                                -> /Module/_Verse.VerseAsset-some_verse_class
 	// 
 	// Our ImportTable will be transformed from the left side to the right where each line is its own ImportEntry outered to the import "up and to the left" of it:
 	// 
-	// /Script/SrcEngineModule                                                                          /Script/SrcEngineModule                                                           
-	//                        .Default__SrcTypeA										                                       .Default__OnlySubTypeChanged							     
-	//                                         .SrcPropertyToChangeA					                                                                   .DstPropertyToChangeB		     
-	//						  				   .MovedButNotRenamedProperty			          	         					                               .UnchangedProperty    
-	//					      .Default__OnlySubTypeChanged							          	         				       .Default__DstTypeB										   
-	//                                                    .SrcPropertyToChangeB		          	                                                 .MovedButNotRenamedProperty	          
-	//					                                  .UnchangedProperty                  -->       /Script/DstEngineModule				       
-	//					      .Default__SrcTypeB										                				       .Default__DstTypeA
-	//					                        .MovedButNotRenamedProperty	          			         				                         .DstPropertyToChangeA
-	//					      .Default__MovedButNotRenamedType						          	         				                         .MovedButNotRenamedProperty
-	//					                                      .MovedButNotRenamedProperty               				       .Default__MovedButNotRenamedType						     
-	//                                                                                                                                                         .MovedButNotRenamedProperty
-
+    // /Script/SrcEngineModule                                                                                          /Script/SrcEngineModule
+    //                        .SrcTypeA                                                                                                       .OnlySubTypeChanged                                 
+    //                                 .SrcPropertyToChangeA                                                                                                     .DstPropertyToChangeB             
+    //                                  .MovedButNotRenamedProperty                                                                                                                   .UnchangedProperty    
+    //                        .OnlySubTypeChanged                                                                                             .DstTypeB                                           
+    //                                           .SrcPropertyToChangeB                                                                                 .MovedButNotRenamedProperty              
+    //                                           .UnchangedProperty                                               -->                         .NewOuter
+    //                        .SrcTypeB                                                                                                                .MovedToNewOuter
+    //                                 .MovedButNotRenamedProperty                                                      /Script/DstEngineModule                       
+	//                                 .MovedToNewOuter                                        							                       .DstTypeA
+    //                        .MovedButNotRenamedType                                                                  	                                .DstPropertyToChangeA
+    //                                               .MovedButNotRenamedProperty                                                                        .MovedButNotRenamedProperty                                               
+    // /Module/_Verse/VerseAsset                                                                                       	                                .MovedButNotRenamedType                             
+    //                          .some_verse_class                                                                      	                                                       .MovedButNotRenamedProperty
+    //                                           .__verse_0x7A8CDEBC_VerseObject1                                      	/Module/_Verse         
+    //                                                                           .__verse_0x5614AC82_VerseObject2      	              .VerseAsset-some_verse_class         
+    //                                                                                                                 	                                          .__verse_0x7A8CDEBC_VerseObject1
+    //                                                                                                                 	                                                                          .__verse_0x5614AC82_VerseObject2
 	SECTION("FObjectResource Patching")
 	{
 		SECTION("Redirect object to new package keeps the original package name if still in use")
@@ -4470,16 +4539,19 @@ TEST_CASE_NAMED(FAssetHeaderPatcherTests, "AssetHeaderPatcher", "[AssetHeaderPat
 			ResetPatcher();
 			TArray<FCoreRedirect> ImportTableRedirects =
 			{
-				{ ECoreRedirectFlags::Type_Package, TEXT("/Engine/SrcClassPackage"),																		TEXT("/Engine/DstClassPackage")},
-				{ ECoreRedirectFlags::Type_Class, TEXT("/Engine/SrcClassPackage.SrcClass"),																	TEXT("/Engine/DstClassPackage.DstClass")},
-				{ ECoreRedirectFlags::Type_Class, TEXT("/Script/SrcEngineModule.Default__SrcTypeA"),														TEXT("/Script/DstEngineModule.Default__DstTypeA")},
-				{ ECoreRedirectFlags::Type_Class, TEXT("/Script/SrcEngineModule.Default__SrcTypeA.SrcPropertyToChangeA"),									TEXT("/Script/DstEngineModule.Default__DstTypeA.DstPropertyToChangeA")},
-				{ ECoreRedirectFlags::Type_Class, TEXT("/Script/SrcEngineModule.Default__SrcTypeA.MovedButNotRenamedProperty"),								TEXT("/Script/DstEngineModule.Default__DstTypeA.MovedButNotRenamedProperty")},
-				{ ECoreRedirectFlags::Type_Class, TEXT("/Script/SrcEngineModule.Default__OnlySubTypeChanged.SrcPropertyToChangeB"),							TEXT("/Script/SrcEngineModule.Default__OnlySubTypeChanged.DstPropertyToChangeB")},
-				{ ECoreRedirectFlags::Type_Class, TEXT("/Script/SrcEngineModule.Default__SrcTypeB"),														TEXT("/Script/SrcEngineModule.Default__DstTypeB")},
-				{ ECoreRedirectFlags::Type_Class, TEXT("/Script/SrcEngineModule.Default__SrcTypeB.MovedButNotRenamedProperty"),								TEXT("/Script/SrcEngineModule.Default__DstTypeB.MovedButNotRenamedProperty")},
-				{ ECoreRedirectFlags::Type_Class, TEXT("/Script/SrcEngineModule.Default__MovedButNotRenamedType"),											TEXT("/Script/DstEngineModule.Default__MovedButNotRenamedType")},
-				{ ECoreRedirectFlags::Type_Class, TEXT("/Script/SrcEngineModule.Default__MovedButNotRenamedType.MovedButNotRenamedProperty"),				TEXT("/Script/DstEngineModule.Default__MovedButNotRenamedType.MovedButNotRenamedProperty")}
+				{ ECoreRedirectFlags::Type_Package, TEXT("/Engine/SrcClassPackage"),                                                   TEXT("/Engine/DstClassPackage")},
+				{ ECoreRedirectFlags::Type_Class,   TEXT("/Engine/SrcClassPackage.SrcClass"),                                          TEXT("/Engine/DstClassPackage.DstClass")},
+				{ ECoreRedirectFlags::Type_Object,  TEXT("/Script/SrcEngineModule.SrcTypeA"),                                          TEXT("/Script/DstEngineModule.DstTypeA")},
+				{ ECoreRedirectFlags::Type_Object,  TEXT("/Script/SrcEngineModule.SrcTypeA.SrcPropertyToChangeA"),                     TEXT("/Script/DstEngineModule.DstTypeA.DstPropertyToChangeA")},
+				{ ECoreRedirectFlags::Type_Object,  TEXT("/Script/SrcEngineModule.SrcTypeA.MovedButNotRenamedProperty"),               TEXT("/Script/DstEngineModule.DstTypeA.MovedButNotRenamedProperty")},
+				{ ECoreRedirectFlags::Type_Object,  TEXT("/Script/SrcEngineModule.OnlySubTypeChanged.SrcPropertyToChangeB"),           TEXT("/Script/SrcEngineModule.OnlySubTypeChanged.DstPropertyToChangeB")},
+				{ ECoreRedirectFlags::Type_Object,  TEXT("/Script/SrcEngineModule.SrcTypeB"),                                          TEXT("/Script/SrcEngineModule.DstTypeB")},
+				{ ECoreRedirectFlags::Type_Object,  TEXT("/Script/SrcEngineModule.SrcTypeB.MovedButNotRenamedProperty"),               TEXT("/Script/SrcEngineModule.DstTypeB.MovedButNotRenamedProperty")},
+				{ ECoreRedirectFlags::Type_Object,  TEXT("/Script/SrcEngineModule.SrcTypeB.MovedToNewOuter"),                          TEXT("/Script/SrcEngineModule.NewOuter.MovedToNewOuter")},
+				{ ECoreRedirectFlags::Type_Object,  TEXT("/Script/SrcEngineModule.MovedButNotRenamedType"),                            TEXT("/Script/DstEngineModule.MovedButNotRenamedType")},
+				{ ECoreRedirectFlags::Type_Object,  TEXT("/Script/SrcEngineModule.MovedButNotRenamedType.MovedButNotRenamedProperty"), TEXT("/Script/DstEngineModule.MovedButNotRenamedType.MovedButNotRenamedProperty")},
+				{ ECoreRedirectFlags::Type_Package, TEXT("/Module/_Verse/VerseAsset"),                                                 TEXT("/Module/_Verse") },
+				{ ECoreRedirectFlags::Type_Object,  TEXT("/Module/_Verse/VerseAsset.some_verse_class"),                                TEXT("/Module/_Verse.VerseAsset-some_verse_class") }
 			};
 			CHECK(FCoreRedirects::AddRedirectList(ImportTableRedirects, TEXT("Asset Header Patcher Tests - FObjectResource Patching")));
 
