[Revision Control] Fix an issue where renaming files resulted in no revision control history for the new, moved file

diff --git a/Engine/Plugins/Developer/PerforceSourceControl/Source/PerforceSourceControl/Private/PerforceSourceControlState.cpp b/Engine/Plugins/Developer/PerforceSourceControl/Source/PerforceSourceControl/Private/PerforceSourceControlState.cpp
--- a/Engine/Plugins/Developer/PerforceSourceControl/Source/PerforceSourceControl/Private/PerforceSourceControlState.cpp
+++ b/Engine/Plugins/Developer/PerforceSourceControl/Source/PerforceSourceControl/Private/PerforceSourceControlState.cpp
@@ -274,8 +274,10 @@ bool FPerforceSourceControlState::CanCheckout() const
 		// Notice that we don't care whether we're up to date. User can perform textual
 		// merge via P4V:
 		const bool bIsCheckedOutElseWhere = State == EPerforceState::CheckedOutOther;
+		const bool bIsOpenForBranch = State == EPerforceState::Branched;
 		bCanDoCheckout =	bIsInP4NotCheckedOut ||
-							bIsCheckedOutElseWhere;
+							bIsCheckedOutElseWhere ||
+							bIsOpenForBranch;
 	}
 	else
 	{
diff --git a/Engine/Source/Developer/AssetTools/Private/AssetRenameManager.cpp b/Engine/Source/Developer/AssetTools/Private/AssetRenameManager.cpp
--- a/Engine/Source/Developer/AssetTools/Private/AssetRenameManager.cpp
+++ b/Engine/Source/Developer/AssetTools/Private/AssetRenameManager.cpp
@@ -1898,23 +1898,6 @@ void FAssetRenameManager::PerformAssetRename(TArray<FAssetRenameDataWithReferenc
 		}
 	}
 
-	RenamingSlowTask.EnterProgressFrame(0.5f, LOCTEXT("SavingRenamedAssetSlowTask", "Saving renamed assets..."));
-
-	// Save all renamed assets and any redirectors that were left behind
-	if (PackagesToSave.Num() > 0)
-	{
-		const bool bCheckDirty = false;
-		const bool bPromptToSave = false;
-		const bool bAlreadyCheckedOut = true;
-
-		// Get the list of filenames before calling save because some of the saved packages can get GCed if they are empty packages
-		const TArray<FString> Filenames = USourceControlHelpers::PackageFilenames(PackagesToSave);
-
-		FEditorFileUtils::PromptForCheckoutAndSave(PackagesToSave, bCheckDirty, bPromptToSave, nullptr, bAlreadyCheckedOut);
-
-		ISourceControlModule::Get().QueueStatusUpdate(Filenames);
-	}
-
 	const FText& UpdatingRevisionControlSlowTask = LOCTEXT("UpdatingRevisionControlSlowTask", "Updating Revision Control...");
 
 	// Bulk update SCC status for old packages since it is faster than doing it one by one below
@@ -1928,11 +1911,7 @@ void FAssetRenameManager::PerformAssetRename(TArray<FAssetRenameDataWithReferenc
 		for (const FAssetRenameDataWithReferencers& RenameData : AssetsToRename)
 		{
 			UPackage* OldPackage = FindPackage(nullptr, *RenameData.OldObjectPath.GetLongPackageName());
-			UPackage* NewPackage = FindPackage(nullptr, *RenameData.NewObjectPath.GetLongPackageName());
-
-			// If something went wrong when saving and the new asset does not exist on disk, don't branch it
-			// as it will just create a copy and any attempt to load it will result in crashes.
-			if (!RenameData.bOnlyFixSoftReferences && NewPackage && FPackageName::DoesPackageExist(NewPackage->GetName()))
+			if (!RenameData.bOnlyFixSoftReferences)
 			{
 				const FString SourceFilename = USourceControlHelpers::PackageFilename(OldPackage);
 				AllSourceFilenamesToCheck.Add(SourceFilename);
@@ -1950,23 +1929,22 @@ void FAssetRenameManager::PerformAssetRename(TArray<FAssetRenameDataWithReferenc
 		RenamingSlowTask.EnterProgressFrame(0.5f, FText::GetEmpty());
 	}
 
+	// Now branch the files in source control if possible
 	if (ISourceControlModule::Get().IsEnabled())
 	{
-		RenamingSlowTask.EnterProgressFrame(5.0f, UpdatingRevisionControlSlowTask);
-		FScopedSlowTask BuildingRelationshipSlowTask(5.0f, LOCTEXT("BranchingRenamedAssetsSlowTask", "Building a relationship between the renamed files in Revision Control..."));
+		RenamingSlowTask.EnterProgressFrame(1.0f, UpdatingRevisionControlSlowTask);
+		FScopedSlowTask BuildingRelationshipSlowTask(1.0f, LOCTEXT("BranchingRenamedAssetsSlowTask", "Building a relationship between the renamed files in Revision Control..."));
 
-		// Now branch the files in source control if possible
 		ProgressStep = 2.9f / static_cast<float>(AssetsToRename.Num());
+		TMap<UPackage*, UPackage> PackagesToBranch;
+		TArray<FString> LongPackageNamesToCheckout;
 		for (const FAssetRenameDataWithReferencers& RenameData : AssetsToRename)
 		{
 			BuildingRelationshipSlowTask.EnterProgressFrame(ProgressStep);
 
 			UPackage* OldPackage = FindPackage(nullptr, *RenameData.OldObjectPath.GetLongPackageName());
 			UPackage* NewPackage = FindPackage(nullptr, *RenameData.NewObjectPath.GetLongPackageName());
-
-			// If something went wrong when saving and the new asset does not exist on disk, don't branch it
-			// as it will just create a copy and any attempt to load it will result in crashes.
-			if (!RenameData.bOnlyFixSoftReferences && NewPackage && FPackageName::DoesPackageExist(NewPackage->GetName()))
+			if (!RenameData.bOnlyFixSoftReferences)
 			{
 				ISourceControlProvider& SourceControlProvider = ISourceControlModule::Get().GetProvider();
 				const FString SourceFilename = USourceControlHelpers::PackageFilename(OldPackage);
@@ -1977,14 +1955,53 @@ void FAssetRenameManager::PerformAssetRename(TArray<FAssetRenameDataWithReferenc
 					if (!SourceControlState->IsAdded())
 					{
 						SourceControlHelpers::BranchPackage(NewPackage, OldPackage, EStateCacheUsage::Use);
+						// Also mark the file for edit so that it can be saved later
+						LongPackageNamesToCheckout.Add(*RenameData.NewObjectPath.GetLongPackageName());
 					}
 				}
 			}
+			SourceControlHelpers::CheckOutFiles(LongPackageNamesToCheckout, true);
 		}
 	}
 	else
 	{
-		RenamingSlowTask.EnterProgressFrame(5.0f, FText::GetEmpty());
+		RenamingSlowTask.EnterProgressFrame(1.0f, FText::GetEmpty());
+	}
+
+	// Save all renamed assets and any redirectors that were left behind
+	{
+		RenamingSlowTask.EnterProgressFrame(4.0f, LOCTEXT("SavingRenamedAssetSlowTask", "Saving renamed assets..."));
+
+		TArray<UPackage*> FailedPackages;
+		if (PackagesToSave.Num() > 0)
+		{
+			const bool bCheckDirty = false;
+			const bool bPromptToSave = false;
+			const bool bAlreadyCheckedOut = true;
+
+			// Get the list of filenames before calling save because some of the saved packages can get GCed if they are empty packages
+			const TArray<FString> Filenames = USourceControlHelpers::PackageFilenames(PackagesToSave);
+
+			FEditorFileUtils::PromptForCheckoutAndSave(PackagesToSave, bCheckDirty, bPromptToSave, &FailedPackages, bAlreadyCheckedOut);
+
+			ISourceControlModule::Get().QueueStatusUpdate(Filenames);
+		}
+
+		// Revert any previously branched files that did not actually save
+		if (ISourceControlModule::Get().IsEnabled())
+		{
+			RenamingSlowTask.EnterProgressFrame(9.0f, UpdatingRevisionControlSlowTask);
+
+			TArray<FString> FilesToRevert = USourceControlHelpers::PackageFilenames(FailedPackages);
+			if (FilesToRevert.Num() > 0)
+			{
+				USourceControlHelpers::RevertFiles(FilesToRevert);
+			}
+		}
+		else
+		{
+			RenamingSlowTask.EnterProgressFrame(9.0f, FText::GetEmpty());
+		}
 	}
 
 	// Clean up all packages that were left empty
