[AssetTools]

Avoid flushing all loads when requesting a sync or revert of a subset of files compared to all the ones that might be loading.

#rnx

diff --git a/Engine/Source/Developer/AssetTools/Private/AssetViewUtils.cpp b/Engine/Source/Developer/AssetTools/Private/AssetViewUtils.cpp
--- a/Engine/Source/Developer/AssetTools/Private/AssetViewUtils.cpp
+++ b/Engine/Source/Developer/AssetTools/Private/AssetViewUtils.cpp
@@ -1837,6 +1837,34 @@ void AssetViewUtils::SyncPackagesFromSourceControl(const TArray<FString>& Packag
 	}
 }
 
+static void FullyLoadAndResetLoaderForPackages(TArray<FString>& InPackageNamesToLoad, TArray<UPackage*>& OutLoadedPackages)
+{
+	TArray<UObject*, TInlineAllocator<32>> LoadedObjects;
+	TArray<int32, TInlineAllocator<32>> PendingPackageRequestIds;
+
+	LoadedObjects.Reserve(InPackageNamesToLoad.Num());
+	OutLoadedPackages.Reserve(InPackageNamesToLoad.Num());
+	PendingPackageRequestIds.Reserve(InPackageNamesToLoad.Num());
+	for (const FString& PackageName : InPackageNamesToLoad)
+	{
+		UPackage* Package = FindPackage(nullptr, *PackageName);
+		if (Package)
+		{
+			LoadedObjects.Emplace(Package);
+			OutLoadedPackages.Emplace(Package);
+
+			if (!Package->IsFullyLoaded())
+			{
+				PendingPackageRequestIds.Emplace(LoadPackageAsync(PackageName));
+			}
+		}
+	}
+
+	// Detach the linkers of any loaded packages so that SCC can overwrite the files...
+	FlushAsyncLoading(PendingPackageRequestIds);
+	ResetLoaders(LoadedObjects);
+}
+
 bool AssetViewUtils::SyncPackagesFromSourceControl(const TArray<FString>& PackageNames)
 {
 	if (PackageNames.Num() > 0)
@@ -1859,41 +1887,8 @@ bool AssetViewUtils::SyncPackagesFromSourceControl(const TArray<FString>& Packag
 		const TArray<FString> PackageFilenames = SourceControlHelpers::PackageFilenames(PackageNamesToSync);
 
 		// Form a list of loaded packages to reload...
-		TArray<UObject*> LoadedObjects;
 		TArray<UPackage*> LoadedPackages;
-		TArray<UPackage*> PendingPackages;
-		LoadedObjects.Reserve(PackageNamesToSync.Num());
-		LoadedPackages.Reserve(PackageNamesToSync.Num());
-		PendingPackages.Reserve(PackageNamesToSync.Num());
-		for (const FString& PackageName : PackageNamesToSync)
-		{
-			UPackage* Package = FindPackage(nullptr, *PackageName);
-			if (Package)
-			{
-				LoadedObjects.Emplace(Package);
-				LoadedPackages.Emplace(Package);
-
-				if (!Package->IsFullyLoaded())
-				{
-					PendingPackages.Emplace(Package);
-				}
-			}
-		}
-
-		// Detach the linkers of any loaded packages so that SCC can overwrite the files...
-		if (PendingPackages.Num() > 0)
-		{
-			FlushAsyncLoading();
-
-			for (UPackage* Package : PendingPackages)
-			{
-				Package->FullyLoad();
-			}
-		}
-		if (LoadedObjects.Num() > 0)
-		{
-			ResetLoaders(LoadedObjects);
-		}
+		FullyLoadAndResetLoaderForPackages(PackageNamesToSync, LoadedPackages);
 
 		// Sync everything...
 		TSharedRef<FSync> Operation = ISourceControlOperation::Create<FSync>();
@@ -2047,41 +2042,8 @@ static bool SyncPathsFromSourceControl(const FString& Revision, const TArray<FSt
 		}
 
 		// Form a list of loaded packages to reload...
-		TArray<UObject*> LoadedObjects;
 		TArray<UPackage*> LoadedPackages;
-		TArray<UPackage*> PendingPackages;
-		LoadedObjects.Reserve(PackageNamesToSync.Num());
-		LoadedPackages.Reserve(PackageNamesToSync.Num());
-		PendingPackages.Reserve(PackageNamesToSync.Num());
-		for (const FString& PackageName : PackageNamesToSync)
-		{
-			UPackage* Package = FindPackage(nullptr, *PackageName);
-			if (Package)
-			{
-				LoadedObjects.Emplace(Package);
-				LoadedPackages.Emplace(Package);
-
-				if (!Package->IsFullyLoaded())
-				{
-					PendingPackages.Emplace(Package);
-				}
-			}
-		}
-
-		// Detach the linkers of any loaded packages so that SCC can overwrite the files...
-		if (PendingPackages.Num() > 0)
-		{
-			FlushAsyncLoading();
-
-			for (UPackage* Package : PendingPackages)
-			{
-				Package->FullyLoad();
-			}
-		}
-		if (LoadedObjects.Num() > 0)
-		{
-			ResetLoaders(LoadedObjects);
-		}
+		FullyLoadAndResetLoaderForPackages(PackageNamesToSync, LoadedPackages);
 
 		// Group everything...
 		TArray<FString> PathsToSync;
