Added support for pasting Verse paths into the asset property editor if enabled by IAssetTools::GetOnShowingContentVersePath.

diff --git a/Engine/Source/Developer/AssetTools/Private/AssetTools.cpp b/Engine/Source/Developer/AssetTools/Private/AssetTools.cpp
--- a/Engine/Source/Developer/AssetTools/Private/AssetTools.cpp
+++ b/Engine/Source/Developer/AssetTools/Private/AssetTools.cpp
@@ -65,6 +65,7 @@
 #include "Misc/FeedbackContext.h"
 #include "BusyCursor.h"
 #include "AssetExportTask.h"
+#include "Containers/VersePath.h"
 #include "Containers/UnrealString.h"
 #include "Serialization/ArchiveReplaceObjectRef.h"
 #include "AdvancedCopyCustomization.h"
@@ -82,6 +83,7 @@
 #include "Animation/AnimSequence.h"
 #include "UObject/SavePackage.h"
 #include "Dialogs/Dialogs.h"
+#include "VerseVM/VVMVerseClass.h"
 #include "Widgets/Input/SEditableTextBox.h"
 #include "Widgets/Layout/SSpacer.h"
 #include "Dialog/SMessageDialog.h"
@@ -6314,6 +6316,182 @@ bool UAssetToolsImpl::IsAssetReadOnly(const FAssetData& AssetData) const
 	return false;
 }
 
+FAssetData UAssetToolsImpl::FindAssetByVersePath(const UE::Core::FVersePath& VersePath, bool bIncludeOnlyOnDiskAssets, bool bSkipARFilteredAssets) const
+{
+	if (!VersePath.IsValid())
+	{
+		return {};
+	}
+
+	const IAssetRegistry& AssetRegistry = IAssetRegistry::GetChecked();
+
+	UE::AssetRegistry::EEnumerateAssetsFlags EnumerateAssetsFlags =
+		(bIncludeOnlyOnDiskAssets ? UE::AssetRegistry::EEnumerateAssetsFlags::None : UE::AssetRegistry::EEnumerateAssetsFlags::OnlyOnDiskAssets) |
+		(bSkipARFilteredAssets ? UE::AssetRegistry::EEnumerateAssetsFlags::None : UE::AssetRegistry::EEnumerateAssetsFlags::AllowUnfilteredArAssets);
+
+	FAssetData AssetData;
+
+	IPluginManager::Get().EnumerateVersePathMountPoints(VersePath, [&VersePath, bIncludeOnlyOnDiskAssets, bSkipARFilteredAssets, &AssetRegistry, EnumerateAssetsFlags, &AssetData](FStringView MountedAssetPath, FStringView LeafPath)
+	{
+		auto GetAssetByObjectPath = [bIncludeOnlyOnDiskAssets, bSkipARFilteredAssets, &AssetRegistry](const FStringView& PackageName, const FStringView& AssetName) -> FAssetData
+		{
+			if (PackageName.Len() >= NAME_SIZE || AssetName.Len() >= NAME_SIZE)
+			{
+				return {};
+			}
+
+			// If we can't find the package name, it must not exist, no need to pollute with random clipboard strings.
+			FName PackageFName(PackageName, FNAME_Find);
+			if (PackageFName.IsNone())
+			{
+				return {};
+			}
+
+			// If we can't find the asset name, it must not exist, no need to pollute with random clipboard strings.
+			FName AssetFName(AssetName, FNAME_Find);
+			if (AssetFName.IsNone())
+			{
+				return {};
+			}
+
+			return AssetRegistry.GetAssetByObjectPath(
+				FSoftObjectPath::ConstructFromAssetPath(FTopLevelAssetPath(PackageFName, AssetFName)), bIncludeOnlyOnDiskAssets, bSkipARFilteredAssets);
+		};
+
+		// The last part of the leaf path should be the top level object name.
+		FStringView RelativePath;
+		FStringView AssetName;
+		int32 LastSlashIndex = INDEX_NONE;
+		if (LeafPath.FindLastChar(TEXT('/'), LastSlashIndex))
+		{
+			RelativePath = FStringView(LeafPath).Left(LastSlashIndex + 1);
+			AssetName = FStringView(LeafPath).Mid(LastSlashIndex + 1);
+		}
+		else
+		{
+			AssetName = LeafPath;
+		}
+
+		if (AssetName.IsEmpty())
+		{
+			// Can't resolve an empty name, keep checking each plug-in.
+			return true;
+		}
+
+		{
+			// Typically, the short package name and top level object name match, so check that first.
+			FNameBuilder PackageNameBuilder;
+			PackageNameBuilder.Append(MountedAssetPath);
+			PackageNameBuilder.Append(RelativePath);
+			PackageNameBuilder.Append(AssetName);
+
+			FNameBuilder AssetNameBuilder;
+			AssetNameBuilder.Append(AssetName);
+
+			AssetData = GetAssetByObjectPath(PackageNameBuilder, AssetNameBuilder);
+			if (AssetData.IsValid())
+			{
+				// Found a match, stop checking subsequent plug-ins.
+				return false;
+			}
+
+			// It could be that we have a cooked package that only contains compiled blueprints, so check that next.
+			AssetNameBuilder.Append(TEXT("_C"));
+
+			AssetData = GetAssetByObjectPath(PackageNameBuilder, AssetNameBuilder);
+			if (AssetData.IsValid())
+			{
+				// If it is a compiled blueprint, GetVersePath will return a Verse path without the _C.
+				if (VersePath == AssetData.GetVersePath())
+				{
+					// Found a match, stop checking subsequent plug-ins.
+					return false;
+				}
+
+				AssetData = {};
+			}
+		}
+
+		// Since we don't know the short package name, and the Verse path could point to a top level object whose name doesn't
+		// match the short package name, we need to scan the package path for other possible assets.
+		{
+			FNameBuilder PackagePathBuilder;
+			PackagePathBuilder.Append(MountedAssetPath);
+			PackagePathBuilder.Append(RelativePath);
+
+			// Remove trailing '/'.
+			PackagePathBuilder.RemoveSuffix(1);
+
+			if (PackagePathBuilder.Len() < NAME_SIZE)
+			{
+				// If we can't find the package path, it must not exist in the path tree, no need to pollute with random clipboard strings.
+				FName PackagePath(PackagePathBuilder.ToView(), FNAME_Find);
+				if (!PackagePath.IsNone())
+				{
+					FARFilter ARFilter;
+					ARFilter.bIncludeOnlyOnDiskAssets = bIncludeOnlyOnDiskAssets;
+					ARFilter.PackagePaths.Add(PackagePath);
+
+					AssetRegistry.EnumerateAssets(ARFilter, [VerseClassPath = UVerseClass::StaticClass()->GetClassPathName(), &VersePath, &AssetData](const FAssetData& PossibleAssetData)
+					{
+						// We will check Verse classes later on.
+						if (VerseClassPath != PossibleAssetData.AssetClassPath && VersePath == PossibleAssetData.GetVersePath())
+						{
+							AssetData = PossibleAssetData;
+						}
+						return !AssetData.IsValid();
+					}, EnumerateAssetsFlags);
+
+					// Check if EnumerateAssets found a match.
+					if (AssetData.IsValid())
+					{
+						// Found a match, stop checking subsequent plug-ins.
+						return false;
+					}
+				}
+			}
+		}
+
+		// Keep checking each plug-in.
+		return true;
+	});
+
+	// Check if EnumerateVersePathMountPoints found a match.
+	if (!AssetData.IsValid())
+	{
+		// Check if we have any Verse classes that match the Verse path.
+		// These can have arbitrary verse paths since the package Verse path is defined by the UE module (ModuleRules.SetupVerse).
+		FARFilter ARFilter;
+		ARFilter.bIncludeOnlyOnDiskAssets = bIncludeOnlyOnDiskAssets;
+		ARFilter.ClassPaths.Add(UVerseClass::StaticClass()->GetClassPathName());
+
+		AssetRegistry.EnumerateAssets(ARFilter, [&VersePath, &AssetData](const FAssetData& PossibleAssetData)
+		{
+			if (VersePath == PossibleAssetData.GetVersePath())
+			{
+				AssetData = PossibleAssetData;
+			}
+			return !AssetData.IsValid();
+		}, EnumerateAssetsFlags);
+	}
+
+	return AssetData;
+}
+
+bool UAssetToolsImpl::ShowingContentVersePath() const
+{
+	if (ShowingContentVersePathDelegate.IsBound())
+	{
+		return ShowingContentVersePathDelegate.Execute();
+	}
+	return false;
+}
+
+UE::AssetTools::FShowingContentVersePath& UAssetToolsImpl::GetOnShowingContentVersePath()
+{
+	return ShowingContentVersePathDelegate;
+}
+
 bool UAssetToolsImpl::AllPassWritableFolderFilter(const TArray<FString>& InPaths) const
 {
 	if (WritableFolderPermissionList->HasFiltering())
diff --git a/Engine/Source/Developer/AssetTools/Private/AssetTools.h b/Engine/Source/Developer/AssetTools/Private/AssetTools.h
--- a/Engine/Source/Developer/AssetTools/Private/AssetTools.h
+++ b/Engine/Source/Developer/AssetTools/Private/AssetTools.h
@@ -146,6 +146,9 @@ public:
 	virtual TSharedRef<FPathPermissionList>& GetWritableFolderPermissionList() override;
 	virtual bool IsAssetVisible(const FAssetData& AssetData, bool bCheckAliases = true) const override;
 	virtual bool IsAssetReadOnly(const FAssetData& AssetData) const override;
+	virtual FAssetData FindAssetByVersePath(const UE::Core::FVersePath& VersePath, bool bIncludeOnlyOnDiskAssets = false, bool bSkipARFilteredAssets = true) const override;
+	virtual bool ShowingContentVersePath() const override;
+	virtual UE::AssetTools::FShowingContentVersePath& GetOnShowingContentVersePath() override;
 	virtual bool AllPassWritableFolderFilter(const TArray<FString>& InPaths) const override;
 	virtual void NotifyBlockedByWritableFolderFilter() const;
 	virtual bool IsNameAllowed(const FString& Name, FText* OutErrorMessage = nullptr) const override;
@@ -288,6 +291,8 @@ private:
 	bool CreateAssetsAsExternallyReferenceable;
 	UE::AssetTools::FShouldCreateAssetsAsExternallyReferenceableForPath CreateAssetsAsExternallyReferenceableForPathDelegate;
 
+	UE::AssetTools::FShowingContentVersePath ShowingContentVersePathDelegate;
+
 	TMap<FName, FIsNameAllowed> IsNameAllowedDelegates;
 	TMap<FName, FIsObjectPathAllowed> IsObjectPathAllowedDelegates;
 	TMap<FName, FSanitizeName> SanitizeNameDelegates;
diff --git a/Engine/Source/Developer/AssetTools/Public/IAssetTools.h b/Engine/Source/Developer/AssetTools/Public/IAssetTools.h
--- a/Engine/Source/Developer/AssetTools/Public/IAssetTools.h
+++ b/Engine/Source/Developer/AssetTools/Public/IAssetTools.h
@@ -3,6 +3,7 @@
 #pragma once
 
 #include "CoreMinimal.h"
+#include "Containers/VersePathFwd.h"
 #include "UObject/WeakObjectPtr.h"
 #include "UObject/Package.h"
 #include "UObject/ObjectMacros.h"
@@ -40,6 +41,8 @@ namespace UE::AssetTools
 	DECLARE_DELEGATE_RetVal_OneParam(bool, FCanAssetBePublic, FStringView /*AssetPath*/);
 
 	DECLARE_MULTICAST_DELEGATE_TwoParams(FShouldCreateAssetsAsExternallyReferenceableForPath, FStringView /*AssetPath*/, TOptional<bool>& /*ExternallyReferenceable*/);
+
+	DECLARE_DELEGATE_RetVal(bool, FShowingContentVersePath);
 }
 
 UENUM()
@@ -656,6 +659,24 @@ public:
 	UFUNCTION(BlueprintCallable, Category = "Editor Scripting | Asset Tools")
 	virtual bool IsAssetReadOnly(const FAssetData& AssetData) const = 0;
 
+	/**
+	 * Find the asset data for the specified Verse path
+	 * @note Verse paths to objects can be ambiguous. This function can be expensive for certain Verse paths.
+	 *
+	 * @param VersePath the Verse path of the object to be looked up
+	 * @param bIncludeOnlyOnDiskAssets If true, use only DiskGatheredData, do not calculate from UObjects.
+	 *        @see IAssetRegistry class header for bIncludeOnlyOnDiskAssets.
+	 * @param bSkipARFilteredAssets If true, skips Objects that return true for IsAsset but are not assets in the current platform.
+	 * @return the assets data; Will be invalid if object could not be found
+	 */
+	UE_INTERNAL virtual FAssetData FindAssetByVersePath(const UE::Core::FVersePath& VersePath, bool bIncludeOnlyOnDiskAssets = false, bool bSkipARFilteredAssets = true) const = 0;
+
+	/** Return whether to surface Content Verse paths throughout the editor */
+	UE_INTERNAL virtual bool ShowingContentVersePath() const = 0;
+
+	/** Delegate to control whether we are showing content Verse paths or not */
+	UE_INTERNAL virtual UE::AssetTools::FShowingContentVersePath& GetOnShowingContentVersePath() = 0;
+
 	/** Returns true if all in list pass writable folder filter */
 	virtual bool AllPassWritableFolderFilter(const TArray<FString>& InPaths) const = 0;
 
diff --git a/Engine/Source/Editor/PropertyEditor/Private/UserInterface/PropertyEditor/SPropertyEditorAsset.cpp b/Engine/Source/Editor/PropertyEditor/Private/UserInterface/PropertyEditor/SPropertyEditorAsset.cpp
--- a/Engine/Source/Editor/PropertyEditor/Private/UserInterface/PropertyEditor/SPropertyEditorAsset.cpp
+++ b/Engine/Source/Editor/PropertyEditor/Private/UserInterface/PropertyEditor/SPropertyEditorAsset.cpp
@@ -6,6 +6,7 @@
 #include "Engine/Texture.h"
 #include "Engine/SkeletalMesh.h"
 #include "Components/StaticMeshComponent.h"
+#include "Containers/VersePath.h"
 #include "Engine/StaticMesh.h"
 #include "Engine/LevelScriptActor.h"
 #include "Editor.h"
@@ -969,9 +970,10 @@ bool SPropertyEditorAsset::ParseAssetText(const FString& InText, FAssetData& Out
 	// All supported paths start with a /
 	if (PossibleObjectPath[0] == TEXT('/'))
 	{
-		const bool bTextIsExportTextPath = PossibleObjectPath != InText;
+		constexpr bool bIncludeOnlyOnDiskAssets = false;
+		constexpr bool bSkipARFilteredAssets = true;
 
-		const FAssetRegistryModule& AssetRegistryModule = FModuleManager::Get().LoadModuleChecked<FAssetRegistryModule>(TEXT("AssetRegistry"));
+		const bool bTextIsExportTextPath = PossibleObjectPath != InText;
 
 		// Check if we just have a package name.
 		int32 ObjectNameDelimiterIndex;
@@ -984,11 +986,33 @@ bool SPropertyEditorAsset::ParseAssetText(const FString& InText, FAssetData& Out
 
 		if (PossibleObjectPath.Len() < NAME_SIZE)
 		{
-			OutAssetData = AssetRegistryModule.Get().GetAssetByObjectPath(FSoftObjectPath(PossibleObjectPath));
+			OutAssetData = IAssetRegistry::GetChecked().GetAssetByObjectPath(FSoftObjectPath(PossibleObjectPath), bIncludeOnlyOnDiskAssets, bSkipARFilteredAssets);
+			if (OutAssetData.IsValid())
+			{
+				return true;
+			}
+		}
+
+		// If it wasn't an export text, object or package path, it might be a Verse path.
+		if (!bTextIsExportTextPath)
+		{
+			FAssetToolsModule& AssetToolsModule = FAssetToolsModule::GetModule();
+			if (AssetToolsModule.Get().ShowingContentVersePath())
+			{
+				UE::Core::FVersePath VersePath;
+				if (UE::Core::FVersePath::TryMake(VersePath, InText))
+				{
+					OutAssetData = AssetToolsModule.Get().FindAssetByVersePath(VersePath, bIncludeOnlyOnDiskAssets, bSkipARFilteredAssets);
+					if (OutAssetData.IsValid())
+					{
+						return true;
+					}
+				}
+			}
 		}
 	}
 
-	return OutAssetData.IsValid();
+	return false;
 }
 
 TSharedRef<SWidget> SPropertyEditorAsset::OnGetMenuContent()
diff --git a/Engine/Source/Editor/PropertyEditor/PropertyEditor.Build.cs b/Engine/Source/Editor/PropertyEditor/PropertyEditor.Build.cs
--- a/Engine/Source/Editor/PropertyEditor/PropertyEditor.Build.cs
+++ b/Engine/Source/Editor/PropertyEditor/PropertyEditor.Build.cs
@@ -44,6 +44,7 @@ public class PropertyEditor : ModuleRules
 				"ConfigEditor",
                 "SceneOutliner",
 				"DesktopPlatform",
+				"Projects",
 				"PropertyPath",
 				"ToolWidgets",
 				"WidgetRegistration",
diff --git a/Engine/Source/Runtime/Projects/Private/PluginManager.cpp b/Engine/Source/Runtime/Projects/Private/PluginManager.cpp
--- a/Engine/Source/Runtime/Projects/Private/PluginManager.cpp
+++ b/Engine/Source/Runtime/Projects/Private/PluginManager.cpp
@@ -3621,32 +3621,36 @@ FName FPluginManager::PackageNameFromModuleName(FName ModuleName)
 	return Result;
 }
 
-bool FPluginManager::TrySplitVersePath(const UE::Core::FVersePath& VersePath, FName& OutPackageName, FString& OutLeafPath)
+void FPluginManager::EnumerateVersePathMountPoints(const UE::Core::FVersePath& VersePath, TFunctionRef<bool(FStringView /*MountedAssetPath*/, FStringView /*LeafPath*/)> Callback)
 {
 	// Can't do anything with an empty vpath
 	if (!VersePath.IsValid())
 	{
-		return false;
+		return;
 	}
 
-	FStringView VersePathView = VersePath.AsStringView();
+	const FStringView VersePathView = VersePath.AsStringView();
 
 	for (const TPair<FString, TArray<TSharedRef<FPlugin>>>& NamePluginPair : AllPlugins)
 	{
 		const TSharedRef<FPlugin>& Plugin = DiscoveredPluginMapUtils::ResolvePluginFromMapVal(NamePluginPair.Value);
 
-		const FPluginDescriptor& PluginDescriptor = Plugin->Descriptor;
-		if (!PluginDescriptor.VersePath.IsEmpty() && VersePathView.StartsWith(PluginDescriptor.VersePath + TEXT('/')))
+		if (Plugin->IsMounted() && !Plugin->GetVersePath().IsEmpty())
 		{
-			VersePathView.RightChopInline(PluginDescriptor.VersePath.Len());
-
-			OutPackageName = *Plugin->Name;
-			OutLeafPath = VersePathView;
-			return true;
+			UE::Core::FVersePath PluginVersePath;
+			if (UE::Core::FVersePath::TryMake(PluginVersePath, Plugin->GetVersePath()))
+			{
+				FStringView LeafPath;
+				if (PluginVersePath.IsBaseOf(VersePath, &LeafPath))
+				{
+					if (!Callback(Plugin->GetMountedAssetPath(), LeafPath))
+					{
+						break;
+					}
+				}
+			}
 		}
 	}
-
-	return false;
 }
 
 #if WITH_EDITOR
diff --git a/Engine/Source/Runtime/Projects/Private/PluginManager.h b/Engine/Source/Runtime/Projects/Private/PluginManager.h
--- a/Engine/Source/Runtime/Projects/Private/PluginManager.h
+++ b/Engine/Source/Runtime/Projects/Private/PluginManager.h
@@ -227,7 +227,7 @@ public:
 	virtual bool GetPluginDependencies_FromFileName(const FString& PluginFileName, TArray<FPluginReferenceDescriptor>& PluginDependencies) override;
 	virtual bool GetPluginDependencies_FromDescriptor(const FPluginReferenceDescriptor& PluginDescriptor, TArray<FPluginReferenceDescriptor>& PluginDependencies) override;
 	virtual FName PackageNameFromModuleName(FName ModuleName) override;
-	virtual bool TrySplitVersePath(const UE::Core::FVersePath& VersePath, FName& OutPackageName, FString& OutLeafPath) override;
+	virtual void EnumerateVersePathMountPoints(const UE::Core::FVersePath& VersePath, TFunctionRef<bool(FStringView /*MountedAssetPath*/, FStringView /*LeafPath*/)> Callback) override;
 	virtual bool RequiresTempTargetForCodePlugin(const FProjectDescriptor* ProjectDescriptor, const FString& Platform, EBuildConfiguration Configuration, EBuildTargetType TargetType, FText& OutReason) override;
 
 	virtual void SetBinariesRootDirectories(const FString& EngineBinariesRootDir, const FString& ProjectBinariesRootDir) override;
diff --git a/Engine/Source/Runtime/Projects/Public/Interfaces/IPluginManager.h b/Engine/Source/Runtime/Projects/Public/Interfaces/IPluginManager.h
--- a/Engine/Source/Runtime/Projects/Public/Interfaces/IPluginManager.h
+++ b/Engine/Source/Runtime/Projects/Public/Interfaces/IPluginManager.h
@@ -563,9 +563,12 @@ public:
 	virtual FName PackageNameFromModuleName(FName ModuleName) = 0;
 
 	/**
-	* Does a reverse lookup to try to figure out what the package name is from a VersePath
+	* Enumerates all mounted plug-ins to find mount points that the Verse path maps to.
+	* 
+	* @param VersePath The full Verse path to check against the plug-in Verse path.
+	* @param Callback Function to call for each plug-in that maps to the Verse path. Return false to stop enumerating.
 	*/
-	virtual bool TrySplitVersePath(const UE::Core::FVersePath& VersePath, FName& OutPackageName, FString& OutLeafPath) = 0;
+	UE_INTERNAL virtual void EnumerateVersePathMountPoints(const UE::Core::FVersePath& VersePath, TFunctionRef<bool(FStringView /*MountedAssetPath*/, FStringView /*LeafPath*/)> Callback) = 0;
 
 	/**
 	 * Determines if a content-only project requires a temporary target due to having a plugin enabled
