[AssetHeaderPatcher]

Only apply full object path redirects when patching to avoid applying partial matches that might overlap the names of irrelevant objects. While not common, it's possible for non-full path redirects to be defined in code and INI which may cause incorrect patching. Since the patcher only deals with full paths for src->dst transformations, restricting to only full path redirects when matching is fine, allowing us to apply global redirects as well as our explicit renames in tandem.

#rnx

diff --git a/Engine/Source/Developer/AssetTools/Private/AssetHeaderPatcher.cpp b/Engine/Source/Developer/AssetTools/Private/AssetHeaderPatcher.cpp
--- a/Engine/Source/Developer/AssetTools/Private/AssetHeaderPatcher.cpp
+++ b/Engine/Source/Developer/AssetTools/Private/AssetHeaderPatcher.cpp
@@ -1876,7 +1876,7 @@ bool FAssetHeaderPatcherInner::DoPatch(FName& InOutName)
 	// If we are given an FName to patch we have no real context as to what that FName is
 	// so we conservatively assume it is a package path and attempt to patch that only
 	FCoreRedirectObjectName SrcPackageName(NAME_None, NAME_None, InOutName);
-	FCoreRedirectObjectName DstPackageName = FCoreRedirects::GetRedirectedName(ECoreRedirectFlags::Type_Package, SrcPackageName);
+	FCoreRedirectObjectName DstPackageName = FCoreRedirects::GetRedirectedName(ECoreRedirectFlags::Type_Package, SrcPackageName, ECoreRedirectMatchFlags::DisallowPartialLHSMatch);
 	if (RemapFName(SrcPackageName.PackageName, DstPackageName.PackageName))
 	{
 		InOutName = DstPackageName.PackageName;
@@ -1892,7 +1892,7 @@ void FAssetHeaderPatcherInner::GetExportTablePatches(TArray<FExportPatch>& OutEx
 	{
 		FObjectExport& Export = ExportTable[i];
 		FCoreRedirectObjectName SrcResourceName = GetFullObjectNameFromObjectResource(Export, true /*bIsExport*/);
-		FCoreRedirectObjectName DstResourceName = FCoreRedirects::GetRedirectedName(ECoreRedirectFlags::Type_AllMask, SrcResourceName);
+		FCoreRedirectObjectName DstResourceName = FCoreRedirects::GetRedirectedName(ECoreRedirectFlags::Type_AllMask, SrcResourceName, ECoreRedirectMatchFlags::DisallowPartialLHSMatch);
 		if (RemapFName(SrcResourceName.ObjectName, DstResourceName.ObjectName))
 		{
 			FExportPatch ExportPatch;
@@ -1931,7 +1931,7 @@ FAssetHeaderPatcher::EResult FAssetHeaderPatcherInner::GetImportTablePatches(TAr
 
 		// For each FObjectResource, immediately patch the FNames that do not impact other imports or exports.
 		const FCoreRedirectObjectName SrcImportClass(Import.ClassName, NAME_None, Import.ClassPackage);
-		const FCoreRedirectObjectName DstImportClass = FCoreRedirects::GetRedirectedName(ECoreRedirectFlags::Type_Class | ECoreRedirectFlags::Type_Package, SrcImportClass);
+		const FCoreRedirectObjectName DstImportClass = FCoreRedirects::GetRedirectedName(ECoreRedirectFlags::Type_Class | ECoreRedirectFlags::Type_Package, SrcImportClass, ECoreRedirectMatchFlags::DisallowPartialLHSMatch);
 
 		Import.ClassName = RemapFName(SrcImportClass.ObjectName, DstImportClass.ObjectName)
 			? DstImportClass.ObjectName : SrcImportClass.ObjectName;
@@ -1941,14 +1941,14 @@ FAssetHeaderPatcher::EResult FAssetHeaderPatcherInner::GetImportTablePatches(TAr
 
 #if WITH_EDITORONLY_DATA
 		const FCoreRedirectObjectName SrcImportPackageName(NAME_None, NAME_None, Import.PackageName);
-		const FCoreRedirectObjectName DstImportPackageName = FCoreRedirects::GetRedirectedName(ECoreRedirectFlags::Type_Package, SrcImportPackageName);
+		const FCoreRedirectObjectName DstImportPackageName = FCoreRedirects::GetRedirectedName(ECoreRedirectFlags::Type_Package, SrcImportPackageName, ECoreRedirectMatchFlags::DisallowPartialLHSMatch);
 		Import.PackageName = RemapFName(SrcImportPackageName.PackageName, DstImportPackageName.PackageName)
 			? DstImportPackageName.PackageName : SrcImportPackageName.PackageName;
 #endif
 
 		// Look up whether there is a specific redirect for the full path of the import.
 		PatchDataForImport.SrcImportPath = GetFullObjectNameFromObjectResource(Import, false /*bIsExport*/, false /*bWalkImportsOnly*/);
-		PatchDataForImport.DstImportPath = FCoreRedirects::GetRedirectedName(FCoreRedirects::GetFlagsForTypeName(Import.ClassPackage, Import.ClassName), PatchDataForImport.SrcImportPath);
+		PatchDataForImport.DstImportPath = FCoreRedirects::GetRedirectedName(FCoreRedirects::GetFlagsForTypeName(Import.ClassPackage, Import.ClassName), PatchDataForImport.SrcImportPath, ECoreRedirectMatchFlags::DisallowPartialLHSMatch);
 
 		// Record the remapping of the FName in the FObjectResource (whether or not it was redirected)
 		FName SrcObjectResourceName = GetObjectResourceNameFromCoreRedirectObjectName(PatchDataForImport.SrcImportPath);
@@ -2306,7 +2306,7 @@ bool FAssetHeaderPatcherInner::DoPatch(FSoftObjectPath& InOutSoft)
 	// name should be marked as unchanging (causing any future patching to the name to be an add in the NameTable). If we have a redirect 
 	// but not remapping, the failed remapping could marked the srcname as unchanging erroneously.
 	FTopLevelAssetPath SrcTopLevelAssetPath = InOutSoft.GetAssetPath();
-	const FCoreRedirectObjectName DstTopLevelAssetPathObjectName = FCoreRedirects::GetRedirectedName(ECoreRedirectFlags::Type_AllMask, SrcTopLevelAssetPath);
+	const FCoreRedirectObjectName DstTopLevelAssetPathObjectName = FCoreRedirects::GetRedirectedName(ECoreRedirectFlags::Type_AllMask, SrcTopLevelAssetPath, ECoreRedirectMatchFlags::DisallowPartialLHSMatch);
 	FTopLevelAssetPath DstTopLevelAssetPath(DstTopLevelAssetPathObjectName.ToString());
 	bool bPatched = SrcTopLevelAssetPath != DstTopLevelAssetPath;
 	if (bPatched)
@@ -2404,7 +2404,7 @@ FCoreRedirectObjectName FAssetHeaderPatcherInner::GetFullObjectNameFromObjectRes
 
 bool FAssetHeaderPatcherInner::DoPatch(FTopLevelAssetPath& InOutPath)
 {
-	const FCoreRedirectObjectName SrcTopLevelAssetPath = FCoreRedirects::GetRedirectedName(ECoreRedirectFlags::Type_AllMask, InOutPath);
+	const FCoreRedirectObjectName SrcTopLevelAssetPath = FCoreRedirects::GetRedirectedName(ECoreRedirectFlags::Type_AllMask, InOutPath, ECoreRedirectMatchFlags::DisallowPartialLHSMatch);
 	const FTopLevelAssetPath DstTopLevelAssetPath(SrcTopLevelAssetPath.ToString());
 
 	bool bPatched = RemapFName(InOutPath.GetAssetName(), DstTopLevelAssetPath.GetAssetName());
@@ -2443,7 +2443,7 @@ bool FAssetHeaderPatcherInner::DoPatch(FGatherableTextData& InOutGatherableTextD
 		}
 
 		FTopLevelAssetPath TopLevelAssetPath = SiteDescriptionPath.GetAssetPath();
-		const FCoreRedirectObjectName RedirectedTopLevelAssetPath = FCoreRedirects::GetRedirectedName(ECoreRedirectFlags::Type_AllMask, TopLevelAssetPath);
+		const FCoreRedirectObjectName RedirectedTopLevelAssetPath = FCoreRedirects::GetRedirectedName(ECoreRedirectFlags::Type_AllMask, TopLevelAssetPath, ECoreRedirectMatchFlags::DisallowPartialLHSMatch);
 		const FTopLevelAssetPath PatchedTopLevelAssetPath(RedirectedTopLevelAssetPath.ToString());
 		if (TopLevelAssetPath == PatchedTopLevelAssetPath)
 		{
@@ -2475,11 +2475,11 @@ bool FAssetHeaderPatcherInner::DoPatch(FThumbnailEntry& InThumbnailEntry)
 	}
 
 	const FCoreRedirectObjectName SrcTopLevelAssetName(FName(SrcObjectPathWithoutPackageName), NAME_None, PackageFName);
-	const FCoreRedirectObjectName RedirectedTopLevelAssetName = FCoreRedirects::GetRedirectedName(ECoreRedirectFlags::Type_Object, SrcTopLevelAssetName);
+	const FCoreRedirectObjectName RedirectedTopLevelAssetName = FCoreRedirects::GetRedirectedName(ECoreRedirectFlags::Type_Object, SrcTopLevelAssetName, ECoreRedirectMatchFlags::DisallowPartialLHSMatch);
 	bool bPatched = RemapFName(SrcTopLevelAssetName.ObjectName, RedirectedTopLevelAssetName.ObjectName);
 
 	const FCoreRedirectObjectName SrcClassName(FName(InThumbnailEntry.ObjectShortClassName), NAME_None, NAME_None);
-	const FCoreRedirectObjectName RedirectedClassName = FCoreRedirects::GetRedirectedName(ECoreRedirectFlags::Type_Class, SrcClassName);
+	const FCoreRedirectObjectName RedirectedClassName = FCoreRedirects::GetRedirectedName(ECoreRedirectFlags::Type_Class, SrcClassName, ECoreRedirectMatchFlags::DisallowPartialLHSMatch);
 	bPatched |= RemapFName(SrcClassName.ObjectName, RedirectedClassName.ObjectName);
 
 	if (bPatched)
@@ -2496,7 +2496,7 @@ FAssetHeaderPatcher::EResult FAssetHeaderPatcherInner::PatchHeader_PatchSections
 	// Package Summary
 	{
 		const FCoreRedirectObjectName DstPackageName = FCoreRedirects::GetRedirectedName(ECoreRedirectFlags::Type_Package,
-			FCoreRedirectObjectName(NAME_None, NAME_None, OriginalPackagePath));
+			FCoreRedirectObjectName(NAME_None, NAME_None, OriginalPackagePath), ECoreRedirectMatchFlags::DisallowPartialLHSMatch);
 
 		// This is a string, so we do not want to Remap the patched name unless it's a non-OFPA
 		// package, in which case there will be an FName entry for this path
@@ -2543,7 +2543,7 @@ FAssetHeaderPatcher::EResult FAssetHeaderPatcherInner::PatchHeader_PatchSections
 	for (FName& Reference : SoftPackageReferencesTable)
 	{
 		FCoreRedirectObjectName SrcPackagePath(NAME_None, NAME_None, Reference);
-		FCoreRedirectObjectName DstPackageName = FCoreRedirects::GetRedirectedName(ECoreRedirectFlags::Type_Package, SrcPackagePath);
+		FCoreRedirectObjectName DstPackageName = FCoreRedirects::GetRedirectedName(ECoreRedirectFlags::Type_Package, SrcPackagePath, ECoreRedirectMatchFlags::DisallowPartialLHSMatch);
 		bool* UsedInGame = AssetRegistryData.SoftPackageReferenceUsedInGame.Find(SrcPackagePath.PackageName);
 		bool bUsedInGame = UsedInGame ? *UsedInGame : true;
 		AssetRegistryData.SoftPackageReferenceUsedInGame.Add(DstPackageName.PackageName, bUsedInGame);
@@ -2674,7 +2674,7 @@ FAssetHeaderPatcher::EResult FAssetHeaderPatcherInner::PatchHeader_PatchSections
 				if (TagData.Value == OriginalPrimaryAssetName)
 				{
 					const FCoreRedirectObjectName DstPackageName = FCoreRedirects::GetRedirectedName(ECoreRedirectFlags::Type_Package,
-						FCoreRedirectObjectName(NAME_None, NAME_None, OriginalPackagePath));
+						FCoreRedirectObjectName(NAME_None, NAME_None, OriginalPackagePath), ECoreRedirectMatchFlags::DisallowPartialLHSMatch);
 
 					TStringBuilder<256> Builder;
 					DstPackageName.PackageName.ToString(Builder);
@@ -2715,7 +2715,7 @@ FAssetHeaderPatcher::EResult FAssetHeaderPatcherInner::PatchHeader_PatchSections
 		{
 			FCoreRedirectObjectName SrcPackagePath(NAME_None, NAME_None, Pair.Key);
 			FCoreRedirectObjectName DstPackageName =
-				FCoreRedirects::GetRedirectedName(ECoreRedirectFlags::Type_Package, SrcPackagePath);
+				FCoreRedirects::GetRedirectedName(ECoreRedirectFlags::Type_Package, SrcPackagePath, ECoreRedirectMatchFlags::DisallowPartialLHSMatch);
 			if (RemapFName(SrcPackagePath.PackageName, DstPackageName.PackageName))
 			{
 				AddedKeys.Add({ DstPackageName.PackageName, Pair.Value});
