[Assets] Fix incorrect values for FSlowTask in rename assets function

#rnx

diff --git a/Engine/Source/Developer/AssetTools/Private/AssetRenameManager.cpp b/Engine/Source/Developer/AssetTools/Private/AssetRenameManager.cpp
--- a/Engine/Source/Developer/AssetTools/Private/AssetRenameManager.cpp
+++ b/Engine/Source/Developer/AssetTools/Private/AssetRenameManager.cpp
@@ -1932,8 +1932,8 @@ void FAssetRenameManager::PerformAssetRename(TArray<FAssetRenameDataWithReferenc
 	// Now branch the files in source control if possible
 	if (ISourceControlModule::Get().IsEnabled())
 	{
-		RenamingSlowTask.EnterProgressFrame(1.0f, UpdatingRevisionControlSlowTask);
-		FScopedSlowTask BuildingRelationshipSlowTask(1.0f, LOCTEXT("BranchingRenamedAssetsSlowTask", "Building a relationship between the renamed files in Revision Control..."));
+		RenamingSlowTask.EnterProgressFrame(2.9f, UpdatingRevisionControlSlowTask);
+		FScopedSlowTask BuildingRelationshipSlowTask(2.9f, LOCTEXT("BranchingRenamedAssetsSlowTask", "Building a relationship between the renamed files in Revision Control..."));
 
 		ProgressStep = 2.9f / static_cast<float>(AssetsToRename.Num());
 		TMap<UPackage*, UPackage> PackagesToBranch;
@@ -1965,7 +1965,7 @@ void FAssetRenameManager::PerformAssetRename(TArray<FAssetRenameDataWithReferenc
 	}
 	else
 	{
-		RenamingSlowTask.EnterProgressFrame(1.0f, FText::GetEmpty());
+		RenamingSlowTask.EnterProgressFrame(2.9f, FText::GetEmpty());
 	}
 
 	// Save all renamed assets and any redirectors that were left behind
