Show Epic internal asset menu entry now that it has been tested
#rnx

diff --git a/Engine/Source/Developer/AssetTools/Private/AssetTools.cpp b/Engine/Source/Developer/AssetTools/Private/AssetTools.cpp
--- a/Engine/Source/Developer/AssetTools/Private/AssetTools.cpp
+++ b/Engine/Source/Developer/AssetTools/Private/AssetTools.cpp
@@ -467,7 +467,7 @@ namespace UE::AssetTools::Private
 		TEXT("Enables the Experimental Public Asset Feature (False: disabled, True:enabled")
 	);
 
-	static bool bEnableEpicInternalAssetFeature = false;
+	static bool bEnableEpicInternalAssetFeature = true;
 	static FAutoConsoleVariableRef CVarEnableEpicInternalAssetFeature(
 		TEXT("AssetTools.EnableEpicInternalAssetFeature"),
 		bEnableEpicInternalAssetFeature,
