[Assets] Fix incorrect values for FSlowTask in rename assets function

#rnx

diff --git a/Engine/Source/Developer/AssetTools/Private/AssetRenameManager.cpp b/Engine/Source/Developer/AssetTools/Private/AssetRenameManager.cpp
--- a/Engine/Source/Developer/AssetTools/Private/AssetRenameManager.cpp
+++ b/Engine/Source/Developer/AssetTools/Private/AssetRenameManager.cpp
@@ -1965,7 +1965,7 @@ void FAssetRenameManager::PerformAssetRename(TArray<FAssetRenameDataWithReferenc
 	}
 	else
 	{
-		RenamingSlowTask.EnterProgressFrame(2.9f, FText::GetEmpty());
+		RenamingSlowTask.EnterProgressFrame(1.0f, FText::GetEmpty());
 	}
 
 	// Save all renamed assets and any redirectors that were left behind
