Ensure AssetTools does not call slate functions in headless.

While running the editor with a script using the command line parameter `-script="C:\path\to\script.py"`, `FApp::IsUnattended()` returns false.
Functions within `UAssetToolsImpl` are still trying to display dialogs which ends up crashing the Editor.
This CL adds a function `UAssetToolsImpl::CanShowDialogs()` which checks if the Application can currently show dialogs (meaning the App is attended, Slate is initialized and we are not running an unattended or automation scripts).

Test script:
```json
import unreal

asset1 = unreal.AssetToolsHelpers.get_asset_tools().create_asset("BP_PSAsset", "/Game", unreal.PaperSprite, unreal.PaperSpriteFactory())
asset2 = unreal.AssetToolsHelpers.get_asset_tools().create_asset("BP_PSAsset", "/Game", unreal.PaperSprite, unreal.PaperSpriteFactory())
```

diff --git a/Engine/Source/Developer/AssetTools/Private/AssetTools.cpp b/Engine/Source/Developer/AssetTools/Private/AssetTools.cpp
--- a/Engine/Source/Developer/AssetTools/Private/AssetTools.cpp
+++ b/Engine/Source/Developer/AssetTools/Private/AssetTools.cpp
@@ -734,7 +734,7 @@ namespace UE::AssetTools::Private
 					{ 
 						// Handle name collision in the destination project. Post 5.1 should expose the possibility to rename the migrated package
 						EAppReturnType::Type Response;
-						if (FApp::IsUnattended() || !PackageMigrationImplContext.Options.bPrompt || LastResponse == EAppReturnType::YesAll || LastResponse == EAppReturnType::NoAll)
+						if (!UAssetToolsImpl::CanShowDialogs() || !PackageMigrationImplContext.Options.bPrompt || LastResponse == EAppReturnType::YesAll || LastResponse == EAppReturnType::NoAll)
 						{
 							Response = LastResponse;
 						}
@@ -2117,6 +2117,11 @@ void UAssetToolsImpl::GenerateAdvancedCopyDestinations(FAdvancedCopyParams& InPa
 	}
 }
 
+bool UAssetToolsImpl::CanShowDialogs()
+{
+	return !FApp::IsUnattended() && FSlateApplication::IsInitialized() && !GIsRunningUnattendedScript && !IsRunningCommandlet();
+}
+
 FString UAssetToolsImpl::GenerateAdvancedCopyDestinationPackageName(const FString& SourcePackage, const FString& SourcePath, const FString& DestinationFolder)
 {
 	FString DestinationPackageName;
@@ -3243,7 +3248,7 @@ void UAssetToolsImpl::MigratePackages(const TArray<FName>& PackageNamesToMigrate
 void UAssetToolsImpl::MigratePackages(const TArray<FName>& PackageNamesToMigrate, const FString& DestinationPath, const struct FMigrationOptions& Options) const
 {
 	// Packages must be saved for the migration to work
-	const bool bPromptUserToSave = !FApp::IsUnattended() && Options.bPrompt;
+	const bool bPromptUserToSave = CanShowDialogs() && Options.bPrompt;
 	const bool bSaveMapPackages = true;
 	const bool bSaveContentPackages = true;
 	if (FEditorFileUtils::SaveDirtyPackages(bPromptUserToSave, bSaveMapPackages, bSaveContentPackages))
@@ -4636,8 +4641,15 @@ bool UAssetToolsImpl::CanCreateAsset(const FString& AssetName, const FString& Pa
 	UObject* ExistingObject = StaticFindObject( UObject::StaticClass(), Pkg, *AssetName );
 	if( ExistingObject != nullptr )
 	{
-		// Object already exists in either the specified package or another package.  Check to see if the user wants
-		// to replace the object.
+		// Object already exists in either the specified package or another package.
+
+		if (!CanShowDialogs())
+		{
+			UE_LOG(LogAssetTools, Error, TEXT("The asset '%s' already exists in package '%s'. CanCreateAsset cannot ask the user as the application is running unattended and will return false"), *AssetName, *PackageName);
+			return false;
+		}
+		
+		// Check to see if the user wants to replace the object.
 		TSharedRef<SMessageDialog> ConfirmDialog = SNew(SMessageDialog)
 			.Icon(FAppStyle::Get().GetBrush("Icons.WarningWithColor.Large"))
 			.Title(FText(NSLOCTEXT("UnrealEd", "ReplaceExistingObjectInPackageConfirmation_Title", "Overwrite Existing Object")))
@@ -4836,7 +4848,7 @@ void UAssetToolsImpl::PerformMigratePackages(TArray<FName> PackageNamesToMigrate
 	// Confirm that there is at least one package to move 
 	if (AllPackageNamesToMove.Num() == 0)
 	{
-		if (!FApp::IsUnattended())
+		if (CanShowDialogs())
 		{
 			FMessageDialog::Open(EAppMsgType::Ok, LOCTEXT("MigratePackages_NoFilesToMove", "No files were found to move"));
 		}
@@ -4850,7 +4862,7 @@ void UAssetToolsImpl::PerformMigratePackages(TArray<FName> PackageNamesToMigrate
 		ReportPackages.Get()->Add({ (*PackageIt).ToString(), bShouldMigratePackage });
 	}
 	// Prompt the user displaying all assets that are going to be migrated
-	if(!FApp::IsUnattended() && Options.bPrompt)
+	if(CanShowDialogs() && Options.bPrompt)
 	{
 		const FText ReportMessage = LOCTEXT("MigratePackagesReportTitle", "The following assets will be migrated to another content folder.");
 		SPackageReportDialog::FOnReportConfirmed OnReportConfirmed = SPackageReportDialog::FOnReportConfirmed::CreateUObject(this, &UAssetToolsImpl::MigratePackages_ReportConfirmed, ReportPackages, DestinationPath, MoveTemp(ExcludedDependencies), MoveTemp(PackageToExternalObjectPackages), Options);
@@ -4865,7 +4877,7 @@ void UAssetToolsImpl::PerformMigratePackages(TArray<FName> PackageNamesToMigrate
 void UAssetToolsImpl::MigratePackages_ReportConfirmed(TSharedPtr<TArray<ReportPackageData>> PackageDataToMigrate, const FString DestinationPath, TSet<FName> ExcludedDependencies, TMap<FName, TSet<FName>> PackageToExternalObjectPackages, const FMigrationOptions Options) const
 {
 	FString DestinationFolder;
-	if (FApp::IsUnattended() || !DestinationPath.IsEmpty())
+	if (!CanShowDialogs() || !DestinationPath.IsEmpty())
 	{
 		if(DestinationPath.IsEmpty())
 		{
@@ -5208,7 +5220,7 @@ void UAssetToolsImpl::MigratePackages_ReportConfirmed(TSharedPtr<TArray<ReportPa
 		// If some packages don't have a matching content root in the destination, prompt for desired behavior
 		if (!LostPackages.IsEmpty())
 		{
-			if (!FApp::IsUnattended() && Options.bPrompt && Options.OrphanFolder.IsEmpty())
+			if (CanShowDialogs() && Options.bPrompt && Options.OrphanFolder.IsEmpty())
 			{
 				FString LostPackageRootsString;
 				for (const FName& PackageRoot : LostPackageRoots)
@@ -5244,7 +5256,7 @@ void UAssetToolsImpl::MigratePackages_ReportConfirmed(TSharedPtr<TArray<ReportPa
 		{
 			// Query the user for a folder to migrate assets into. This folder will exist temporary in this project, and is the destination for migration in the target project
 			FString FolderName = Options.OrphanFolder.IsEmpty() ? TEXT("Migrated") : Options.OrphanFolder;
-			if (!FApp::IsUnattended() && Options.bPrompt)
+			if (CanShowDialogs() && Options.bPrompt)
 			{
 				bool bIsOkButtonEnabled = true;
 				TSharedRef<SEditableTextBox> EditableTextBox = SNew(SEditableTextBox)
@@ -5319,7 +5331,7 @@ void UAssetToolsImpl::MigratePackages_ReportConfirmed(TSharedPtr<TArray<ReportPa
 			if (IFileManager::Get().DirectoryExists(*SrcDiskFolderFilename))
 			{
 				const FText Message = FText::Format(LOCTEXT("MigratePackages_InvalidMigrateFolder", "{0} exists on disk in the source project, and cannot be used to consolidate assets."), FText::FromString(SrcDiskFolderFilename));
-				if (!FApp::IsUnattended() && Options.bPrompt)
+				if (CanShowDialogs() && Options.bPrompt)
 				{
 					EAppReturnType::Type Response = FMessageDialog::Open(EAppMsgType::Ok, Message);
 				}
@@ -5528,7 +5540,7 @@ void UAssetToolsImpl::MigratePackages_ReportConfirmed(TSharedPtr<TArray<ReportPa
 						{
 							// The destination file already exists! Ask the user what to do.
 							EAppReturnType::Type Response;
-							if (FApp::IsUnattended() || !Options.bPrompt || LastResponse == EAppReturnType::YesAll || LastResponse == EAppReturnType::NoAll)
+							if (!CanShowDialogs() || !Options.bPrompt || LastResponse == EAppReturnType::YesAll || LastResponse == EAppReturnType::NoAll)
 							{
 								Response = LastResponse;
 							}
@@ -5929,7 +5941,7 @@ void UAssetToolsImpl::PerformAdvancedCopyPackages(TArray<FName> SelectedAssetAnd
 	CopyParams.OnCopyComplete = OnCopyComplete;
 
 	// Suppress UI if we're running in unattended mode
-	if (FApp::IsUnattended())
+	if (!CanShowDialogs())
 	{
 		CopyParams.bShouldSuppressUI = true;
 	}
diff --git a/Engine/Source/Developer/AssetTools/Private/AssetTools.h b/Engine/Source/Developer/AssetTools/Private/AssetTools.h
--- a/Engine/Source/Developer/AssetTools/Private/AssetTools.h
+++ b/Engine/Source/Developer/AssetTools/Private/AssetTools.h
@@ -183,6 +183,9 @@ public:
 
 	/** The manager to handle fixing up redirectors */
 	TSharedPtr<FAssetFixUpRedirectors> AssetFixUpRedirectors;
+
+	/** Check if the Application can currently show dialogs, meaning the App is attended, Slate is initialized and we are not running an unattended or automation scripts */
+	static bool CanShowDialogs();
 private:
 	static FString GenerateAdvancedCopyDestinationPackageName(const FString& SourcePackage, const FString& SourcePath, const FString& DestinationFolder);
 
