[Assets] Fix incorrect values for FSlowTask in rename assets function when moving assets between folders

#rnx

diff --git a/Engine/Source/Developer/AssetTools/Private/AssetRenameManager.cpp b/Engine/Source/Developer/AssetTools/Private/AssetRenameManager.cpp
--- a/Engine/Source/Developer/AssetTools/Private/AssetRenameManager.cpp
+++ b/Engine/Source/Developer/AssetTools/Private/AssetRenameManager.cpp
@@ -2000,7 +2000,7 @@ void FAssetRenameManager::PerformAssetRename(TArray<FAssetRenameDataWithReferenc
 		}
 		else
 		{
-			RenamingSlowTask.EnterProgressFrame(9.0f, FText::GetEmpty());
+			RenamingSlowTask.EnterProgressFrame(5.0f, FText::GetEmpty());
 		}
 	}
 
