Added support for showing Verse paths to the package report dialog.
Fixed an issue where a parent node would show the checked state if all its children were undetermined.
Fixed unexpanded folders appearing with the open folder icon instead of the closed one.

diff --git a/Engine/Source/Developer/AssetTools/Private/SPackageReportDialog.cpp b/Engine/Source/Developer/AssetTools/Private/SPackageReportDialog.cpp
--- a/Engine/Source/Developer/AssetTools/Private/SPackageReportDialog.cpp
+++ b/Engine/Source/Developer/AssetTools/Private/SPackageReportDialog.cpp
@@ -2,6 +2,7 @@
 
 
 #include "SPackageReportDialog.h"
+#include "IAssetTools.h"
 #include "Modules/ModuleManager.h"
 #include "Widgets/SWindow.h"
 #include "Layout/WidgetPath.h"
@@ -13,38 +14,89 @@
 #include "Widgets/Input/SButton.h"
 #include "Styling/AppStyle.h"
 #include "Interfaces/IMainFrameModule.h"
+#include "Interfaces/IPluginManager.h"
 
 #define LOCTEXT_NAMESPACE "PackageReportDialog"
 
-struct FCompareFPackageReportNodeByName
-{
-	FORCEINLINE bool operator()( TSharedPtr<FPackageReportNode> A, TSharedPtr<FPackageReportNode> B ) const
-	{
-		return A->NodeName < B->NodeName;
-	}
-};
-
 FPackageReportNode::FPackageReportNode()
 	: CheckedState(ECheckBoxState::Undetermined)
 	, bShouldMigratePackage(nullptr)
-	, bIsFolder(false)
 	, Parent(nullptr)
 {}
 
-FPackageReportNode::FPackageReportNode(const FString& InNodeName, bool InIsFolder)
-	: NodeName(InNodeName)
-	, CheckedState(ECheckBoxState::Undetermined)
-	, bShouldMigratePackage(nullptr)
-	, bIsFolder(InIsFolder)
-	, Parent(nullptr)
-{}
+FPackageReportNode::FPackageReportNode(FString&& InNodeName, FText&& InNodeText, FText&& InToolTipText, bool* bInShouldMigratePackage, FPackageReportNode& Parent)
+	: NodeName(MoveTemp(InNodeName))
+	, NodeText(MoveTemp(InNodeText))
+	, ToolTipText(MoveTemp(InToolTipText))
+	, CheckedState(bInShouldMigratePackage != nullptr ? (*bInShouldMigratePackage ? ECheckBoxState::Checked : ECheckBoxState::Unchecked) : ECheckBoxState::Undetermined)
+	, bShouldMigratePackage(bInShouldMigratePackage)
+	, Parent(&Parent)
+{
+}
 
 void FPackageReportNode::AddPackage(const FString& PackageName, bool* bInShouldMigratePackage)
 {
 	TArray<FString> PathElements;
 	PackageName.ParseIntoArray(PathElements, TEXT("/"), /*InCullEmpty=*/true);
 
-	(void)AddPackage_Recursive(PathElements, bInShouldMigratePackage);
+	if (PathElements.Num() < 2)
+	{
+		return;
+	}
+
+	const IAssetTools& AssetTools = IAssetTools::Get();
+
+	TSharedPtr<IPlugin> Plugin = IPluginManager::Get().FindPlugin(PathElements[0]);
+	bool bFirstElementIsMountPoint = true;
+	if (Plugin.IsValid() && !Plugin->GetVersePath().IsEmpty() && AssetTools.ShowingContentVersePath())
+	{
+		TArray<FString> VersePathElements;
+		Plugin->GetVersePath().ParseIntoArray(VersePathElements, TEXT("/"), /*InCullEmpty=*/true);
+
+		if (ensure(!VersePathElements.IsEmpty()))
+		{
+			bFirstElementIsMountPoint = false;
+
+			// Replace the mount point with the Verse path.
+			PathElements.RemoveAt(0);
+			PathElements.Insert(VersePathElements, 0);
+		}
+	}
+
+	FPackageReportNode* ChildParent = this;
+	for (int32 Index = 0; Index < PathElements.Num(); ++Index)
+	{
+		TSharedPtr<FPackageReportNode>* Child = ChildParent->Children.FindByPredicate([&PathElements, Index](const TSharedPtr<FPackageReportNode>& Child)
+		{
+			return Child->NodeName == PathElements[Index];
+		});
+		if (Child == nullptr)
+		{
+			FString ChildNodeName = MoveTemp(PathElements[Index]);
+
+			FText ChildNodeText;
+			FText ChildToolTipText;
+			bool* bChildShouldMigratePackage = nullptr;
+			if (bFirstElementIsMountPoint && Index == 0 && Plugin.IsValid())
+			{
+				ChildNodeText = FText::FromString(Plugin->GetFriendlyName());
+				ChildToolTipText = FText::FromString(ChildNodeName);
+			}
+			else
+			{
+				ChildNodeText = FText::FromString(ChildNodeName);
+
+				if (Index + 1 == PathElements.Num())
+				{
+					ChildToolTipText = FText::FromString(AssetTools.GetUserFacingLongPackageName(FName(PackageName)));
+					bChildShouldMigratePackage = bInShouldMigratePackage;
+				}
+			}
+
+			Child = &ChildParent->Children.Add_GetRef(MakeShared<FPackageReportNode>(MoveTemp(ChildNodeName), MoveTemp(ChildNodeText), MoveTemp(ChildToolTipText), bChildShouldMigratePackage, *ChildParent));
+		}
+		ChildParent = Child->Get();
+	}
 }
 
 void FPackageReportNode::ExpandChildrenRecursively(const TSharedRef<PackageReportTree>& Treeview)
@@ -56,56 +108,52 @@ void FPackageReportNode::ExpandChildrenRecursively(const TSharedRef<PackageRepor
 	}
 }
 
-FPackageReportNode::FChildrenState FPackageReportNode::AddPackage_Recursive(TArray<FString>& PathElements, bool* bInShouldMigratePackage)
+void FPackageReportNode::Finalize()
 {
-	FChildrenState ChildrenState;
-	ChildrenState.bAnyChildIsChecked = false;
-	ChildrenState.bAllChildrenAreChecked = true;
+	for (const TSharedPtr<FPackageReportNode>& Child : Children)
+	{
+		Child->Finalize();
+	}
 
-	if ( PathElements.Num() > 0 )
+	Children.Sort([](const TSharedPtr<FPackageReportNode>& A, const TSharedPtr<FPackageReportNode>& B)
 	{
-		// Pop the bottom element
-		FString ChildNodeName = PathElements[0];
-		PathElements.RemoveAt(0);
+		return A->NodeText.CompareTo(B->NodeText) < 0;
+	});
+
+	UpdateCheckedStateFromChildren();
+}
 
-		// Try to find a child which uses this folder name
-		TSharedPtr<FPackageReportNode> Child;
-		for ( auto ChildIt = Children.CreateConstIterator(); ChildIt; ++ChildIt )
+void FPackageReportNode::UpdateCheckedStateFromChildren()
+{
+	if (bShouldMigratePackage == nullptr)
+	{
+		int32 NumChecked = 0;
+		for (const TSharedPtr<FPackageReportNode>& Child : Children)
 		{
-			if ( (*ChildIt)->NodeName == ChildNodeName )
+			switch (Child->CheckedState)
 			{
-				Child = (*ChildIt);
+			case ECheckBoxState::Checked:
+				++NumChecked;
 				break;
+			case ECheckBoxState::Undetermined:
+				CheckedState = ECheckBoxState::Undetermined;
+				return;
 			}
 		}
 
-		// If one was not found, create it
-		if ( !Child.IsValid() )
+		if (NumChecked == Children.Num())
 		{
-			const bool bIsAFolder = (PathElements.Num() > 0);
-			int32 ChildIdx = Children.Add( MakeShareable(new FPackageReportNode(ChildNodeName, bIsAFolder)) );
-			Child = Children[ChildIdx];
-			Child.Get()->Parent = this;
-			Children.Sort( FCompareFPackageReportNodeByName() );
+			CheckedState = ECheckBoxState::Checked;
 		}
-
-		if ( ensure(Child.IsValid()) )
+		else if (NumChecked == 0)
 		{
-			FChildrenState ChildChildrenState = Child->AddPackage_Recursive(PathElements, bInShouldMigratePackage);
-			ChildrenState.bAnyChildIsChecked |= ChildChildrenState.bAnyChildIsChecked;
-			ChildrenState.bAllChildrenAreChecked &= ChildChildrenState.bAllChildrenAreChecked;
+			CheckedState = ECheckBoxState::Unchecked;
+		}
+		else
+		{
+			CheckedState = ECheckBoxState::Undetermined;
 		}
-
-		CheckedState = ChildrenState.bAllChildrenAreChecked ? ECheckBoxState::Checked : (ChildrenState.bAnyChildIsChecked ? ECheckBoxState::Undetermined : ECheckBoxState::Unchecked);
-	}
-	else
-	{
-		CheckedState = *bInShouldMigratePackage ? ECheckBoxState::Checked : ECheckBoxState::Unchecked;
-		ChildrenState.bAnyChildIsChecked = ChildrenState.bAllChildrenAreChecked = CheckedState == ECheckBoxState::Checked;
-		bShouldMigratePackage = bInShouldMigratePackage;
 	}
-
-	return ChildrenState;
 }
 
 BEGIN_SLATE_FUNCTION_BUILD_OPTIMIZATION
@@ -224,8 +272,6 @@ TSharedRef<ITableRow> SPackageReportDialog::GenerateTreeRow( TSharedPtr<FPackage
 {
 	check(TreeItem.IsValid());
 
-	const FSlateBrush* IconBrush = GetNodeIcon(TreeItem);
-
 	return SNew( STableRow< TSharedPtr<FPackageReportNode> >, OwnerTable )
 		[
 			// Icon
@@ -240,13 +286,15 @@ TSharedRef<ITableRow> SPackageReportDialog::GenerateTreeRow( TSharedPtr<FPackage
 			+SHorizontalBox::Slot()
 			.AutoWidth()
 			[
-				SNew(SImage).Image(IconBrush)
+				SNew(SImage).Image(this, &SPackageReportDialog::GetNodeIcon, TreeItem)
 			]
 			// Name
 			+SHorizontalBox::Slot()
 			.FillWidth(1.f)
 			[
-				SNew(STextBlock).Text(FText::FromString(TreeItem->NodeName))
+				SNew(STextBlock)
+				.Text(TreeItem->NodeText)
+				.ToolTipText(TreeItem->ToolTipText)
 				.ColorAndOpacity(FSlateColor::UseForeground())
 			]
 		];
@@ -291,15 +339,7 @@ void SPackageReportDialog::CheckBoxStateChanged(ECheckBoxState InCheckBoxState,
 	FPackageReportNode* CurrentParent = TreeItem->Parent;
 	while (CurrentParent != nullptr)
 	{
-		bool bAnyChildIsChecked = false;
-		bool bAllChildrenAreChecked = true;
-		for (int i = 0; i < CurrentParent->Children.Num(); i++)
-		{
-			bAnyChildIsChecked |= CurrentParent->Children[i]->CheckedState != ECheckBoxState::Unchecked;
-			bAllChildrenAreChecked &= CurrentParent->Children[i]->CheckedState != ECheckBoxState::Unchecked;
-		}
-
-		CurrentParent->CheckedState = bAllChildrenAreChecked ? ECheckBoxState::Checked : (bAnyChildIsChecked ? ECheckBoxState::Undetermined : ECheckBoxState::Unchecked);
+		CurrentParent->UpdateCheckedStateFromChildren();
 		CurrentParent = CurrentParent->Parent;
 	}
 
@@ -317,11 +357,12 @@ void SPackageReportDialog::ConstructNodeTree(TArray<ReportPackageData>& PackageN
 	{
 		PackageReportRootNode.AddPackage(Package.Name, &Package.bShouldMigratePackage);
 	}
+	PackageReportRootNode.Finalize();
 }
 
-const FSlateBrush* SPackageReportDialog::GetNodeIcon(const TSharedPtr<FPackageReportNode>& ReportNode) const
+const FSlateBrush* SPackageReportDialog::GetNodeIcon(TSharedPtr<FPackageReportNode> ReportNode) const
 {
-	if ( !ReportNode->bIsFolder )
+	if ( ReportNode->bShouldMigratePackage != nullptr )
 	{
 		return PackageBrush;
 	}
diff --git a/Engine/Source/Developer/AssetTools/Private/SPackageReportDialog.h b/Engine/Source/Developer/AssetTools/Private/SPackageReportDialog.h
--- a/Engine/Source/Developer/AssetTools/Private/SPackageReportDialog.h
+++ b/Engine/Source/Developer/AssetTools/Private/SPackageReportDialog.h
@@ -23,13 +23,15 @@ struct ReportPackageData
 struct FPackageReportNode
 {
 	/** The name of the tree node without the path */
-	FString NodeName; 
+	FString NodeName;
+	/** The display name of the tree node */
+	FText NodeText;
+	/** The tooltip of the tree node */
+	FText ToolTipText;
 	/** A user-exposed flag determining whether the content of this node and its children should be migrated or not. */
 	ECheckBoxState CheckedState;
 	/** A pointer to an external bool describing whether this node ultimately should be migrated or not. Is only non-null for leaf nodes.*/
 	bool* bShouldMigratePackage;
-	/** If true, this node is a folder instead of a package */
-	bool bIsFolder;
 	
 	/** The parent of this node */
 	FPackageReportNode* Parent;
@@ -39,7 +41,7 @@ struct FPackageReportNode
 
 	/** Constructor */
 	FPackageReportNode();
-	FPackageReportNode(const FString& InNodeName, bool InIsFolder);
+	FPackageReportNode(FString&& InNodeName, FText&& InNodeText, FText&& InToolTipText, bool* bInShouldMigratePackage, FPackageReportNode& Parent);
 
 	/** Adds the path to the tree relative to this node, creating nodes as needed. */
 	void AddPackage(const FString& PackageName, bool* bInIsPackageIncluded);
@@ -47,14 +49,11 @@ struct FPackageReportNode
 	/** Expands this node and all its children */
 	void ExpandChildrenRecursively(const TSharedRef<PackageReportTree>& Treeview);
 
-private:
-	struct FChildrenState
-	{
-		bool bAnyChildIsChecked;
-		bool bAllChildrenAreChecked;
-	};
-	/** Helper function for AddPackage. PathElements is the tokenized path delimited by "/" */
-	FChildrenState AddPackage_Recursive(TArray<FString>& PathElements, bool* bInIsPackageIncluded);
+	/** Call after all nodes have been created to order the child nodes for display and calculate their initial CheckedState. */
+	void Finalize();
+
+	/** Updates CheckedState from this node's child CheckedStates if bShouldMigratePackage is null */
+	void UpdateCheckedStateFromChildren();
 };
 
 class SPackageReportDialog : public SCompoundWidget
@@ -95,7 +94,7 @@ private:
 	void GetChildrenForTree( TSharedPtr<FPackageReportNode> TreeItem, TArray< TSharedPtr<FPackageReportNode> >& OutChildren );
 
 	/** Determines which image to display next to a node */
-	const FSlateBrush* GetNodeIcon(const TSharedPtr<FPackageReportNode>& ReportNode) const;
+	const FSlateBrush* GetNodeIcon(TSharedPtr<FPackageReportNode> ReportNode) const;
 
 	/** Handler for when "Ok" is clicked */
 	FReply OkClicked();
