Added support for showing Verse paths to the Redirector Update Report.
Added the asset path as a tie breaker when sorting items.

diff --git a/Engine/Source/Developer/AssetTools/Private/AssetFixUpRedirectors.cpp b/Engine/Source/Developer/AssetTools/Private/AssetFixUpRedirectors.cpp
--- a/Engine/Source/Developer/AssetTools/Private/AssetFixUpRedirectors.cpp
+++ b/Engine/Source/Developer/AssetTools/Private/AssetFixUpRedirectors.cpp
@@ -225,23 +225,33 @@ public:
 		Sizing = ESizingRule::UserSized;
 		MinDimensions = FVector2D(1000.0f, 600.f);
 
+		const UAssetToolsImpl& AssetTools = UAssetToolsImpl::Get();
+
 		bool bAnyRedirectorsToDelete = false;
 		for (const FRedirectorRefs& Redirector : RedirectorRefsList)
 		{
 			TSharedPtr<FFixupReportItem>& NewItem = TreeItems.Emplace_GetRef(MakeShared<FFixupReportItem>(
 				EFixupReportItemType::Redirector,
-				FText::FromName(Redirector.RedirectorPackageName)));
-			for (FName LockedPackageName : Redirector.LockedReferencerPackageNames)
+				FText::FromString(Redirector.Redirector ? AssetTools.GetUserFacingPathName(*Redirector.Redirector) : AssetTools.GetUserFacingLongPackageName(Redirector.RedirectorPackageName))));
 			{
-				NewItem->Children.Emplace(MakeShared<FFixupReportItem>(
-					EFixupReportItemType::LockedReferencer,
-					FText::FromName(LockedPackageName)));
+				TArray<FString> LockedReferencers = AssetTools.GetUserFacingLongPackageNames(Redirector.LockedReferencerPackageNames);
+				LockedReferencers.Sort();
+				for (FString& LockedReferencer : LockedReferencers)
+				{
+					NewItem->Children.Emplace(MakeShared<FFixupReportItem>(
+						EFixupReportItemType::LockedReferencer,
+						FText::FromString(MoveTemp(LockedReferencer))));
+				}
 			}
-			for (FName FailedPackageName : Redirector.FailedReferencerPackageNames)
 			{
-				NewItem->Children.Emplace(MakeShared<FFixupReportItem>(
-					EFixupReportItemType::FailedReferencer,
-					FText::FromName(FailedPackageName)));
+				TArray<FString> FailedReferencers = AssetTools.GetUserFacingLongPackageNames(Redirector.FailedReferencerPackageNames);
+				FailedReferencers.Sort();
+				for (FString& FailedReferencer : FailedReferencers)
+				{
+					NewItem->Children.Emplace(MakeShared<FFixupReportItem>(
+						EFixupReportItemType::FailedReferencer,
+						FText::FromString(MoveTemp(FailedReferencer))));
+				}
 			}
 			for (FText Issue : Redirector.OtherFailures)
 			{
@@ -271,10 +281,17 @@ public:
 			MinDimensions.Y = 800.0f; // Slightly bigger default size to account for help text
 		}
 
-		Algo::SortBy(
-			TreeItems,
-			[](const TSharedPtr<FFixupReportItem>& Item) -> int32 { return !!Item->Children.Num(); },
-			TGreater<int32>{});
+		TreeItems.Sort([](const TSharedPtr<FFixupReportItem>& A, const TSharedPtr<FFixupReportItem>& B)
+		{
+			// Sort nodes with problems first.
+			const bool bAHasProblems = A->Type != EFixupReportItemType::UnreferencedRedirector;
+			const bool bBHasProblems = B->Type != EFixupReportItemType::UnreferencedRedirector;
+			if (bAHasProblems != bBHasProblems)
+			{
+				return bAHasProblems;
+			}
+			return A->Text.ToString() < B->Text.ToString();
+		});
 
 		TSharedRef<SHeaderRow> HeaderRowWidget =
 			SNew(SHeaderRow)
