[Assets] Fix incorrect values for FSlowTask in rename assets function

diff --git a/Engine/Source/Developer/AssetTools/Private/AssetRenameManager.cpp b/Engine/Source/Developer/AssetTools/Private/AssetRenameManager.cpp
--- a/Engine/Source/Developer/AssetTools/Private/AssetRenameManager.cpp
+++ b/Engine/Source/Developer/AssetTools/Private/AssetRenameManager.cpp
@@ -1965,12 +1965,12 @@ void FAssetRenameManager::PerformAssetRename(TArray<FAssetRenameDataWithReferenc
 	}
 	else
 	{
-		RenamingSlowTask.EnterProgressFrame(1.0f, FText::GetEmpty());
+		RenamingSlowTask.EnterProgressFrame(2.9f, FText::GetEmpty());
 	}
 
 	// Save all renamed assets and any redirectors that were left behind
 	{
-		RenamingSlowTask.EnterProgressFrame(4.0f, LOCTEXT("SavingRenamedAssetSlowTask", "Saving renamed assets..."));
+		RenamingSlowTask.EnterProgressFrame(1.0f, LOCTEXT("SavingRenamedAssetSlowTask", "Saving renamed assets..."));
 
 		TArray<UPackage*> FailedPackages;
 		if (PackagesToSave.Num() > 0)
@@ -1990,7 +1990,7 @@ void FAssetRenameManager::PerformAssetRename(TArray<FAssetRenameDataWithReferenc
 		// Revert any previously branched files that did not actually save
 		if (ISourceControlModule::Get().IsEnabled())
 		{
-			RenamingSlowTask.EnterProgressFrame(9.0f, UpdatingRevisionControlSlowTask);
+			RenamingSlowTask.EnterProgressFrame(1.0f, UpdatingRevisionControlSlowTask);
 
 			TArray<FString> FilesToRevert = USourceControlHelpers::PackageFilenames(FailedPackages);
 			if (FilesToRevert.Num() > 0)
@@ -2000,7 +2000,7 @@ void FAssetRenameManager::PerformAssetRename(TArray<FAssetRenameDataWithReferenc
 		}
 		else
 		{
-			RenamingSlowTask.EnterProgressFrame(5.0f, FText::GetEmpty());
+			RenamingSlowTask.EnterProgressFrame(1.0f, FText::GetEmpty());
 		}
 	}
 
