Fixed typo

#rnx

diff --git a/Engine/Source/Developer/AssetTools/Private/AssetTools.cpp b/Engine/Source/Developer/AssetTools/Private/AssetTools.cpp
--- a/Engine/Source/Developer/AssetTools/Private/AssetTools.cpp
+++ b/Engine/Source/Developer/AssetTools/Private/AssetTools.cpp
@@ -1131,7 +1131,7 @@ namespace UE::AssetTools::Private
 					}
 					else
 					{
-						PackageMigrationContext.AddErrorMigrationMessage(FText::Format(LOCTEXT("MigratePackages_SaveFailedNoAsset", "Couldn't migrate package ({0}) because the package didn't contains a an asset. Destination File ({1})")
+						PackageMigrationContext.AddErrorMigrationMessage(FText::Format(LOCTEXT("MigratePackages_SaveFailedNoAsset", "Couldn't migrate package ({0}) because the package didn't contain an asset. Destination File ({1})")
 							, FText::FromString(MigrationPackageData.GetOriginalPackageName())
 							, FText::FromString(MigrationPackageData.GetDestinationFilename())
 						));
