Removed problematic ensure

diff --git a/Engine/Source/Developer/AssetTools/Private/LocalizedAssetTools.cpp b/Engine/Source/Developer/AssetTools/Private/LocalizedAssetTools.cpp
--- a/Engine/Source/Developer/AssetTools/Private/LocalizedAssetTools.cpp
+++ b/Engine/Source/Developer/AssetTools/Private/LocalizedAssetTools.cpp
@@ -187,7 +187,7 @@ ELocalizedAssetsOnDiskResult FLocalizedAssetTools::GetLocalizedVariantsOnDisk(co
 
 	TMap<FName, FAssetData> PackagesToAssetDataMap;
 	UE::AssetRegistry::GetAssetForPackages(InPackages, PackagesToAssetDataMap);
-	if (!ensureAlwaysMsgf(PackagesToAssetDataMap.Num() == InPackages.Num(), TEXT("PackageNames were not properly converted to AssetData.")))
+	if (PackagesToAssetDataMap.Num() != InPackages.Num()) // There used to be an ensure here but there were edge cases not covered. It will be properly fixed soon.
 	{
 		for (const FName& OriginalAssetName : InPackages)
 		{
