Added support for showing Verse paths to invalid external reference messages.

diff --git a/Engine/Source/Developer/AssetTools/Private/AssetTools.cpp b/Engine/Source/Developer/AssetTools/Private/AssetTools.cpp
--- a/Engine/Source/Developer/AssetTools/Private/AssetTools.cpp
+++ b/Engine/Source/Developer/AssetTools/Private/AssetTools.cpp
@@ -6503,6 +6503,26 @@ UE::AssetTools::FShowingContentVersePath& UAssetToolsImpl::GetOnShowingContentVe
 	return ShowingContentVersePathDelegate;
 }
 
+const FString& UAssetToolsImpl::GetUserFacingPluginName(const IPlugin& Plugin) const
+{
+	if (!Plugin.GetVersePath().IsEmpty() && ShowingContentVersePath())
+	{
+		return Plugin.GetVersePath();
+	}
+
+	return Plugin.GetFriendlyName();
+}
+
+FString UAssetToolsImpl::GetUserFacingPluginName(FStringView PluginName) const
+{
+	if (TSharedPtr<IPlugin> Plugin = IPluginManager::Get().FindPlugin(PluginName))
+	{
+		return GetUserFacingPluginName(*Plugin);
+	}
+
+	return FString(PluginName);
+}
+
 FString UAssetToolsImpl::GetUserFacingLongPackagePath(FStringView LongPackagePath) const
 {
 	if (!LongPackagePath.IsEmpty() && ShowingContentVersePath())
diff --git a/Engine/Source/Developer/AssetTools/Private/AssetTools.h b/Engine/Source/Developer/AssetTools/Private/AssetTools.h
--- a/Engine/Source/Developer/AssetTools/Private/AssetTools.h
+++ b/Engine/Source/Developer/AssetTools/Private/AssetTools.h
@@ -149,6 +149,8 @@ public:
 	virtual FAssetData FindAssetByVersePath(const UE::Core::FVersePath& VersePath, bool bIncludeOnlyOnDiskAssets = false, bool bSkipARFilteredAssets = true) const override;
 	virtual bool ShowingContentVersePath() const override;
 	virtual UE::AssetTools::FShowingContentVersePath& GetOnShowingContentVersePath() override;
+	virtual const FString& GetUserFacingPluginName(const IPlugin& Plugin) const override;
+	virtual FString GetUserFacingPluginName(FStringView PluginName) const override;
 	virtual FString GetUserFacingLongPackagePath(FStringView LongPackagePath) const override;
 	virtual FString GetUserFacingLongPackageName(FName PackageName) const override;
 	virtual TArray<FString> GetUserFacingLongPackageNames(TConstArrayView<FName> PackageNames) const override;
diff --git a/Engine/Source/Developer/AssetTools/Public/IAssetTools.h b/Engine/Source/Developer/AssetTools/Public/IAssetTools.h
--- a/Engine/Source/Developer/AssetTools/Public/IAssetTools.h
+++ b/Engine/Source/Developer/AssetTools/Public/IAssetTools.h
@@ -23,6 +23,7 @@ class IAssetTools;
 class IAssetTypeActions;
 class IClassTypeActions;
 class ILocalizedAssetTools;
+class IPlugin;
 class UFactory;
 class UAssetImportTask;
 class UAdvancedCopyCustomization;
@@ -677,6 +678,18 @@ public:
 	/** Delegate to control whether we are showing content Verse paths or not */
 	UE_INTERNAL virtual UE::AssetTools::FShowingContentVersePath& GetOnShowingContentVersePath() = 0;
 
+	/**
+	 * Return the path to display to the user in editor UI for the given plugin.
+	 * Attempts to convert the folder to a Verse path if ShowingContentVersePath is true.
+	 */
+	UE_INTERNAL virtual const FString& GetUserFacingPluginName(const IPlugin& Plugin) const = 0;
+
+	/**
+	 * Return the path to display to the user in editor UI for the given plugin.
+	 * Attempts to convert the folder to a Verse path if ShowingContentVersePath is true.
+	 */
+	UE_INTERNAL virtual FString GetUserFacingPluginName(FStringView PluginName) const = 0;
+
 	/**
 	 * Return the path to display to the user in editor UI for the given package path.
 	 * Attempts to convert the folder to a Verse path if ShowingContentVersePath is true.
