Show editor progress bar when bulk importing audio files.

diff --git a/Engine/Source/Developer/AssetTools/Private/AssetTools.cpp b/Engine/Source/Developer/AssetTools/Private/AssetTools.cpp
--- a/Engine/Source/Developer/AssetTools/Private/AssetTools.cpp
+++ b/Engine/Source/Developer/AssetTools/Private/AssetTools.cpp
@@ -3334,15 +3334,6 @@ TArray<UObject*> UAssetToolsImpl::ImportAssetsInternal(const TArray<FString>& Fi
 		return UInterchangeManager::IsInterchangeImportEnabled() && (SpecifiedFactory == nullptr);
 	}();
 
-	FScopedSlowTask SlowTask(static_cast<float>(ValidFiles.Num()), LOCTEXT("ImportSlowTask", "Importing"), !bUseInterchangeFramework);
-
-	if (!bUseInterchangeFramework && ValidFiles.Num() > 1)
-	{	
-		//Always allow user to cancel the import task if they are importing multiple ValidFiles.
-		//If we're importing a single file, then the factory policy will dictate if the import if cancelable.
-		SlowTask.MakeDialog(true);
-	}
-
 	TArray<TPair<FString, FString>> FilesAndDestinations;
 	if (FilesAndDestinationsPtr == nullptr)
 	{
@@ -3507,6 +3498,7 @@ TArray<UObject*> UAssetToolsImpl::ImportAssetsInternal(const TArray<FString>& Fi
 
 	TArray<UFactory*> UsedFactories;
 	bool bImportWasCancelled = false;
+	bool bShowSlowTask = !bUseInterchangeFramework;
 	bool bOnlyInterchangeImport = bUseInterchangeFramework;
 	if(bUseInterchangeFramework)
 	{
@@ -3533,16 +3525,30 @@ TArray<UObject*> UAssetToolsImpl::ImportAssetsInternal(const TArray<FString>& Fi
 			{	
 				//Always allow user to cancel the import task if they are importing multiple files.
 				//If we're importing a single file, then the factory policy will dictate if the import if cancelable.
-				SlowTask.MakeDialog(true);
+				bShowSlowTask = true;
 			}
 		}
 		else
 		{
-			//Complete the slow task
-			SlowTask.CompletedWork = static_cast<float>(FilesAndDestinations.Num());
+			//Do not show the slow task if we are only using the Interchange Manager for imports
+			bShowSlowTask = false;
 		}
 	}
 
+	// Create a slow task to document progress
+	FScopedSlowTask SlowTask(static_cast<float>(ValidFiles.Num()), LOCTEXT("ImportSlowTask", "Importing"), bShowSlowTask);
+
+	if (bShowSlowTask && ValidFiles.Num() > 1)
+	{
+		//Always allow user to cancel the import task if they are importing multiple ValidFiles.
+		//If we're importing a single file, then the factory policy will dictate if the import if cancelable.
+		SlowTask.MakeDialog(true);
+	}
+	else
+	{
+		SlowTask.CompletedWork = static_cast<float>(FilesAndDestinations.Num());
+	}
+
 	struct FInterchangeImportStatus
 	{
 		explicit FInterchangeImportStatus(int32 NumFiles)
