[AssetHeaderPatcher]

- Fixes regression from when additional import patching was supported in the patcher. We used to unilaterally generate Export/Import patches even if the name didn't change which allowed us to properly patch FName number portions when they change, despite the name portion not changing. When additional import support was added we made generating Import/Export patches conditional on the FName display id not matching which isn't correct since that ignores the number part.
- Fixes DumpState to properly print the number portion of FNames when dumping so that we can see the NameTable entry FName (always has a 0 number), and the object instance FName (can have a non-0 number)

#rnx

diff --git a/Engine/Source/Developer/AssetTools/Private/AssetHeaderPatcher.cpp b/Engine/Source/Developer/AssetTools/Private/AssetHeaderPatcher.cpp
--- a/Engine/Source/Developer/AssetTools/Private/AssetHeaderPatcher.cpp
+++ b/Engine/Source/Developer/AssetTools/Private/AssetHeaderPatcher.cpp
@@ -1831,7 +1831,9 @@ bool FAssetHeaderPatcherInner::RemapFName(FName SrcName, FName DstName)
 	// Since we do fuzzy matching against AssetRegistryTag data, we might get SrcNames not in the FName table as input; ignore these.
 	if (!NameToIndexMap.Contains(SrcComparisonId))
 	{
-		return false;
+		// If we haven't patched the nametable we still want to flag this name as remapped so 
+		// we can ensure we patch header object instances with any changed numbers
+		return SrcName.GetNumber() != DstName.GetNumber();
 	}
 
 	// Previously it was thought that any FName at the front of the NameTable is an FName used by export data and should not be patched
@@ -1846,7 +1848,10 @@ bool FAssetHeaderPatcherInner::RemapFName(FName SrcName, FName DstName)
 	if (!bNeedRemap)
 	{
 		AddFName(SrcName);
-		return false;
+
+		// If we haven't patched the nametable we still want to flag this name as remapped so 
+		// we can ensure we patch header object instances with any changed numbers
+		return SrcName.GetNumber() != DstName.GetNumber();
 	}
 	else
 	{
@@ -3183,7 +3188,7 @@ void FAssetHeaderPatcherInner::DumpState(FStringView OutputDirectory)
 			if (Index)
 			{
 				FName NameTableName = NameTable[*Index];
-				return FString::Printf(TEXT("%s (nametable index: %d, fname:{'%s', %d})"), *NameTableName.ToString(), *Index, *NameTableName.GetPlainNameString(), NameTableName.GetNumber());
+				return FString::Printf(TEXT("%s (nametable index: %d, fname:{'%s', %d})"), *NameTableName.ToString(), *Index, *Name.GetPlainNameString(), Name.GetNumber());
 			}
 			else
 			{
@@ -3668,6 +3673,8 @@ TEST_CASE_NAMED(FAssetHeaderPatcherTests, "AssetHeaderPatcher", "[AssetHeaderPat
 	const FName DstAssetFName(TEXT("DstPackage"));
 	const FName SrcExportObjectFName = SrcAssetFName;
 	const FName DstExportObjectFName = DstAssetFName;
+	const FName SrcExportSubObjectFName(TEXT("SrcSubObject"));
+	const FName DstExportSubObjectFName(TEXT("DstSubObject"));
 	const FName DummyImportPackagePathFName(TEXT("/DummyMount/DummyPackage"));
 	// Import FNames
 	const FName SrcEngineModuleImportObjectName(TEXT("/Script/SrcEngineModule"));
@@ -3675,7 +3682,7 @@ TEST_CASE_NAMED(FAssetHeaderPatcherTests, "AssetHeaderPatcher", "[AssetHeaderPat
 	const FName SrcTypeBImportObjectName(TEXT("SrcTypeB"));
 	const FName OnlySubTypeChangedImportObjectName(TEXT("OnlySubTypeChanged"));
 	const FName MovedButNotRenamedTypeImportObjectName(TEXT("MovedButNotRenamedType"));
-	const FName SrcPropertyToChangeAImportObjectName(TEXT("SrcPropertyToChangeA"));
+	const FName SrcPropertyToChangeAImportObjectName(TEXT("SrcPropertyToChangeA_2"));
 	const FName SrcPropertyToChangeBImportObjectName(TEXT("SrcPropertyToChangeB"));
 	const FName MovedButNotRenamedPropertyImportObjectName(TEXT("MovedButNotRenamedProperty"));
 	const FName NewOuterImportObjectName(TEXT("NewOuter"));
@@ -3693,7 +3700,7 @@ TEST_CASE_NAMED(FAssetHeaderPatcherTests, "AssetHeaderPatcher", "[AssetHeaderPat
 	const FName DstEngineModuleImportObjectName(TEXT("/Script/DstEngineModule"));
 	const FName DstTypeAImportObjectName(TEXT("DstTypeA"));
 	const FName DstTypeBImportObjectName(TEXT("DstTypeB"));
-	const FName DstPropertyToChangeAImportObjectName(TEXT("DstPropertyToChangeA"));
+	const FName DstPropertyToChangeAImportObjectName(TEXT("DstPropertyToChangeA_4"));
 	const FName DstPropertyToChangeBImportObjectName(TEXT("DstPropertyToChangeB"));
 	const FName DstVerseAssetName(TEXT("/Module/_Verse"));
 	const FName DstVerseImportObjectName(TEXT("VerseAsset-some_verse_class"));
@@ -3711,7 +3718,7 @@ TEST_CASE_NAMED(FAssetHeaderPatcherTests, "AssetHeaderPatcher", "[AssetHeaderPat
 		{ SrcPackagePath, DstPackagePath },
 	};
 
-	auto MakeImport = [SrcPackagePathFName](const FName ObjectName, const FPackageIndex OuterIndex, const FName ClassPackage, const FName ClassName, const FName PackageName)
+	auto MakeImport = [](const FName ObjectName, const FPackageIndex OuterIndex, const FName ClassPackage, const FName ClassName, const FName PackageName)
 		{
 			FObjectImport Import;
 			Import.ObjectName = ObjectName;
@@ -3762,7 +3769,7 @@ TEST_CASE_NAMED(FAssetHeaderPatcherTests, "AssetHeaderPatcher", "[AssetHeaderPat
 			.Dst = MakeImport(DstTypeAImportObjectName, FPackageIndex::FromImport(18), DstImportClassPackage, DstImportClassName, DstPackagePathFName),
 			.bExistingImport = true,
 		},
-		// /Script/SrcEngineModule.SrcTypeA.SrcPropertyToChangeA -> /Script/DstEngineModule.DstTypeA.DstPropertyToChangeA
+		// /Script/SrcEngineModule.SrcTypeA.SrcPropertyToChangeA_2 -> /Script/DstEngineModule.DstTypeA.DstPropertyToChangeA_4
 		{ 
 			.Src = MakeImport(SrcPropertyToChangeAImportObjectName, FPackageIndex::FromImport(1), SrcImportClassPackage, SrcImportClassName, SrcPackagePathFName), 
 			.Dst = MakeImport(DstPropertyToChangeAImportObjectName, FPackageIndex::FromImport(1), DstImportClassPackage, DstImportClassName, DstPackagePathFName),
@@ -3892,20 +3899,20 @@ TEST_CASE_NAMED(FAssetHeaderPatcherTests, "AssetHeaderPatcher", "[AssetHeaderPat
 	TArray<FExportTestCase> ExportTestCases
 	{
 		// SrcPackage																-> DstPackage
-		// Super: /Script/SrcEngineModule											-> /Script/DstEngineModule
-		// Class: /Script/SrcEngineModule.SrcTypeA.SrcPropertyToChangeA				-> /Script/DstEngineModule.DstTypeA.DstPropertyToChangeA
+		// Super: /Script/SrcEngineModule.Package									-> /Script/DstEngineModule.Package
+		// Class: /Script/SrcEngineModule.SrcTypeA.SrcPropertyToChangeA_2			-> /Script/DstEngineModule.DstTypeA.DstPropertyToChangeA_4
 		// Template: /Script/SrcEngineModule.OnlySubTypeChanged						-> /Script/SrcEngineModule.OnlySubTypeChanged
 		{
 			.Src = MakeExport(SrcExportObjectFName, FPackageIndex::FromExport(0), FPackageIndex(), FPackageIndex::FromImport(0), FPackageIndex::FromImport(2), FPackageIndex::FromImport(4)),
 			.Dst = MakeExport(DstExportObjectFName, FPackageIndex::FromExport(0), FPackageIndex(), FPackageIndex::FromImport(11), FPackageIndex::FromImport(2), FPackageIndex::FromImport(4))
 		},
-		// SrcPackage.SrcPackage														-> DstPackage.DstPackage
+		// SrcPackage.SrcSubObject														-> DstPackage.DstSubObject
 		// Super: /Script/SrcEngineModule.SrcTypeA										-> /Script/DstEngineModule.DstTypeA
 		// Class: /Script/SrcEngineModule.SrcTypeA.MovedButNotRenamedProperty			-> /Script/DstEngineModule.DstTypeA.MovedButNotRenamedProperty
 		// Template: /Script/SrcEngineModule.OnlySubTypeChanged.SrcPropertyToChangeB	-> /Script/SrcEngineModule.OnlySubTypeChanged.DstPropertyToChangeB
 		{
-			.Src = MakeExport(SrcExportObjectFName, FPackageIndex::FromExport(1), FPackageIndex::FromExport(0), FPackageIndex::FromImport(1), FPackageIndex::FromImport(3), FPackageIndex::FromImport(5)),
-			.Dst = MakeExport(DstExportObjectFName, FPackageIndex::FromExport(1), FPackageIndex::FromExport(0), FPackageIndex::FromImport(1), FPackageIndex::FromImport(3), FPackageIndex::FromImport(5))
+			.Src = MakeExport(SrcExportSubObjectFName, FPackageIndex::FromExport(1), FPackageIndex::FromExport(0), FPackageIndex::FromImport(1), FPackageIndex::FromImport(3), FPackageIndex::FromImport(5)),
+			.Dst = MakeExport(DstExportSubObjectFName, FPackageIndex::FromExport(1), FPackageIndex::FromExport(0), FPackageIndex::FromImport(1), FPackageIndex::FromImport(3), FPackageIndex::FromImport(5))
 		}
 	};
 
@@ -3946,6 +3953,7 @@ TEST_CASE_NAMED(FAssetHeaderPatcherTests, "AssetHeaderPatcher", "[AssetHeaderPat
 			AddToNameTable(SoftObjectPathToRedirect.GetAssetPath().GetAssetName());
 			// Export Table Names
 			AddToNameTable(SrcExportObjectFName);
+			AddToNameTable(SrcExportSubObjectFName);
 			// Import Table Names
 			AddToNameTable(SrcEngineModuleImportObjectName);
 			AddToNameTable(SrcTypeAImportObjectName);
@@ -4490,7 +4498,7 @@ TEST_CASE_NAMED(FAssetHeaderPatcherTests, "AssetHeaderPatcher", "[AssetHeaderPat
 	// /Engine/SrcClassPackage                                                   -> /Engine/DstClassPackage
 	// /Engine/SrcClassPackage.SrcClass                                          -> /Engine/DstClassPackage.DstClass
 	// /Script/SrcEngineModule.SrcTypeA                                          -> /Script/DstEngineModule.DstTypeA
-	// /Script/SrcEngineModule.SrcTypeA.SrcPropertyToChangeA                     -> /Script/DstEngineModule.DstTypeA.DstPropertyToChangeA
+	// /Script/SrcEngineModule.SrcTypeA.SrcPropertyToChangeA_2                   -> /Script/DstEngineModule.DstTypeA.DstPropertyToChangeA_4
 	// /Script/SrcEngineModule.SrcTypeA.MovedButNotRenamedProperty               -> /Script/DstEngineModule.DstTypeA.MovedButNotRenamedProperty
 	// /Script/SrcEngineModule.OnlySubTypeChanged.SrcPropertyToChangeB           -> /Script/SrcEngineModule.OnlySubTypeChanged.DstPropertyToChangeB
 	// /Script/SrcEngineModule.SrcTypeB                                          -> /Script/SrcEngineModule.DstTypeB
@@ -4505,7 +4513,7 @@ TEST_CASE_NAMED(FAssetHeaderPatcherTests, "AssetHeaderPatcher", "[AssetHeaderPat
 	// 
     // /Script/SrcEngineModule                                                                                          /Script/SrcEngineModule
     //                        .SrcTypeA                                                                                                       .OnlySubTypeChanged                                 
-    //                                 .SrcPropertyToChangeA                                                                                                     .DstPropertyToChangeB             
+    //                                 .SrcPropertyToChangeA_2                                                                                                   .DstPropertyToChangeB             
     //                                  .MovedButNotRenamedProperty                                                                                                                   .UnchangedProperty    
     //                        .OnlySubTypeChanged                                                                                             .DstTypeB                                           
     //                                           .SrcPropertyToChangeB                                                                                 .MovedButNotRenamedProperty              
@@ -4513,7 +4521,7 @@ TEST_CASE_NAMED(FAssetHeaderPatcherTests, "AssetHeaderPatcher", "[AssetHeaderPat
     //                        .SrcTypeB                                                                                                                .MovedToNewOuter
     //                                 .MovedButNotRenamedProperty                                                      /Script/DstEngineModule                       
 	//                                 .MovedToNewOuter                                        							                       .DstTypeA
-    //                        .MovedButNotRenamedType                                                                  	                                .DstPropertyToChangeA
+    //                        .MovedButNotRenamedType                                                                  	                                .DstPropertyToChangeA_4
     //                                               .MovedButNotRenamedProperty                                                                        .MovedButNotRenamedProperty                                               
     // /Module/_Verse/VerseAsset                                                                                       	                                .MovedButNotRenamedType                             
     //                          .some_verse_class                                                                      	                                                       .MovedButNotRenamedProperty
@@ -4532,7 +4540,11 @@ TEST_CASE_NAMED(FAssetHeaderPatcherTests, "AssetHeaderPatcher", "[AssetHeaderPat
 					// should not contain None, so don't look for it
 					if (!Actual.IsNone())
 					{
-						CHECK(Patcher.NameTable.Contains(Actual));
+						FNameEntryId ActualNameEntryId = Actual.GetDisplayIndex();
+						CHECK(Patcher.NameTable.ContainsByPredicate([ActualNameEntryId](const FName& Name)
+							{
+								return Name.GetDisplayIndex() == ActualNameEntryId;
+							}));
 					}
 				};
 
@@ -4542,7 +4554,7 @@ TEST_CASE_NAMED(FAssetHeaderPatcherTests, "AssetHeaderPatcher", "[AssetHeaderPat
 				{ ECoreRedirectFlags::Type_Package, TEXT("/Engine/SrcClassPackage"),                                                   TEXT("/Engine/DstClassPackage")},
 				{ ECoreRedirectFlags::Type_Class,   TEXT("/Engine/SrcClassPackage.SrcClass"),                                          TEXT("/Engine/DstClassPackage.DstClass")},
 				{ ECoreRedirectFlags::Type_Object,  TEXT("/Script/SrcEngineModule.SrcTypeA"),                                          TEXT("/Script/DstEngineModule.DstTypeA")},
-				{ ECoreRedirectFlags::Type_Object,  TEXT("/Script/SrcEngineModule.SrcTypeA.SrcPropertyToChangeA"),                     TEXT("/Script/DstEngineModule.DstTypeA.DstPropertyToChangeA")},
+				{ ECoreRedirectFlags::Type_Object,  TEXT("/Script/SrcEngineModule.SrcTypeA.SrcPropertyToChangeA_2"),                   TEXT("/Script/DstEngineModule.DstTypeA.DstPropertyToChangeA_4")},
 				{ ECoreRedirectFlags::Type_Object,  TEXT("/Script/SrcEngineModule.SrcTypeA.MovedButNotRenamedProperty"),               TEXT("/Script/DstEngineModule.DstTypeA.MovedButNotRenamedProperty")},
 				{ ECoreRedirectFlags::Type_Object,  TEXT("/Script/SrcEngineModule.OnlySubTypeChanged.SrcPropertyToChangeB"),           TEXT("/Script/SrcEngineModule.OnlySubTypeChanged.DstPropertyToChangeB")},
 				{ ECoreRedirectFlags::Type_Object,  TEXT("/Script/SrcEngineModule.SrcTypeB"),                                          TEXT("/Script/SrcEngineModule.DstTypeB")},
@@ -4551,7 +4563,8 @@ TEST_CASE_NAMED(FAssetHeaderPatcherTests, "AssetHeaderPatcher", "[AssetHeaderPat
 				{ ECoreRedirectFlags::Type_Object,  TEXT("/Script/SrcEngineModule.MovedButNotRenamedType"),                            TEXT("/Script/DstEngineModule.MovedButNotRenamedType")},
 				{ ECoreRedirectFlags::Type_Object,  TEXT("/Script/SrcEngineModule.MovedButNotRenamedType.MovedButNotRenamedProperty"), TEXT("/Script/DstEngineModule.MovedButNotRenamedType.MovedButNotRenamedProperty")},
 				{ ECoreRedirectFlags::Type_Package, TEXT("/Module/_Verse/VerseAsset"),                                                 TEXT("/Module/_Verse") },
-				{ ECoreRedirectFlags::Type_Object,  TEXT("/Module/_Verse/VerseAsset.some_verse_class"),                                TEXT("/Module/_Verse.VerseAsset-some_verse_class") }
+				{ ECoreRedirectFlags::Type_Object,  TEXT("/Module/_Verse/VerseAsset.some_verse_class"),                                TEXT("/Module/_Verse.VerseAsset-some_verse_class") },
+				{ ECoreRedirectFlags::Type_Object,  TEXT("/SrcMount/SomePath/SrcPackage.SrcPackage:SrcSubObject"),                     TEXT("/DstMount/SomePath/DstPackage.DstPackage:DstSubObject")},
 			};
 			CHECK(FCoreRedirects::AddRedirectList(ImportTableRedirects, TEXT("Asset Header Patcher Tests - FObjectResource Patching")));
 
