[AssetHeaderPatcher]

Adds tests for FNames that are patched to only have their number change. Verifies that we can patch FObjectResource FNames, as well as all string forms appripriately

#rnx

diff --git a/Engine/Source/Developer/AssetTools/Private/AssetHeaderPatcher.cpp b/Engine/Source/Developer/AssetTools/Private/AssetHeaderPatcher.cpp
--- a/Engine/Source/Developer/AssetTools/Private/AssetHeaderPatcher.cpp
+++ b/Engine/Source/Developer/AssetTools/Private/AssetHeaderPatcher.cpp
@@ -3658,23 +3658,33 @@ TEST_CASE_NAMED(FAssetHeaderPatcherTests, "AssetHeaderPatcher", "[AssetHeaderPat
 
 	// To avoid having to deal with serialization, we mock some data and inject it directly
 	// into the patcher as if done via serialization
-	const FString DummySrcDstAsset			= TEXT("/SrcMount/SomePath/SrcPackage");
-	const TCHAR*  SrcPackagePath			= TEXT("/SrcMount/SomePath/SrcPackage");
-	const TCHAR*  DstPackagePath			= TEXT("/DstMount/SomePath/DstPackage");
-	const TCHAR*  SrcPackageObjectPath		= TEXT("/SrcMount/SomePath/SrcPackage.SrcPackage");
-	const TCHAR*  DstPackageObjectPath		= TEXT("/DstMount/SomePath/DstPackage.DstPackage");
-	const TCHAR*  SrcMountName				= TEXT("/SourceSpecialMount/");
-	const TCHAR*  DstMountName				= TEXT("/DestinationSpecialMount/");
+	const FString DummySrcDstAsset			 = TEXT("/SrcMount/SomePath/SrcPackage");
+	const TCHAR*  SrcPackagePath			 = TEXT("/SrcMount/SomePath/SrcPackage");
+	const TCHAR*  DstPackagePath			 = TEXT("/DstMount/SomePath/DstPackage");
+	const TCHAR*  SrcPackageObjectPath		 = TEXT("/SrcMount/SomePath/SrcPackage.SrcPackage");
+	const TCHAR*  DstPackageObjectPath		 = TEXT("/DstMount/SomePath/DstPackage.DstPackage");
+	const TCHAR*  SrcNumberChangedPath		 = TEXT("/SrcMount/SomePath/UnchangingName_1");
+	const TCHAR*  DstNumberChangedPath		 = TEXT("/DstMount/SomePath/UnchangingName_2");
+	const TCHAR*  SrcNumberChangedObjectPath = TEXT("/SrcMount/SomePath/UnchangingName_1.UnchangingName_1");
+	const TCHAR*  DstNumberChangedObjectPath = TEXT("/DstMount/SomePath/UnchangingName_2.UnchangingName_2");
+	const TCHAR*  SrcMountName				 = TEXT("/SourceSpecialMount/");
+	const TCHAR*  DstMountName				 = TEXT("/DestinationSpecialMount/");
 	const FSoftObjectPath SoftObjectPathToRedirect(TEXT("/ToBeRedirectedMount/SomePath/ToBeRedirectedPackage.ToBeRedirectedPackage:Some.ToBeRedirectedPackage.Subobject"));
 	const FSoftObjectPath RedirectedSoftObjectPath(TEXT("/RedirectedMount/SomePath/RedirectedPackage.RedirectedPackage:Some.RedirectedPackage.Subobject"));
 	const FName SrcPackagePathFName(SrcPackagePath);
 	const FName DstPackagePathFName(DstPackagePath);
+	const FName SrcNumberChangedPathFName(SrcNumberChangedPath);
+	const FName DstNumberChangedPathFName(DstNumberChangedPath);
 	const FName SrcAssetFName(TEXT("SrcPackage"));
 	const FName DstAssetFName(TEXT("DstPackage"));
+	const FName SrcNumberAssetFName(TEXT("UnchangingName_1"));
+	const FName DstNumberAssetFName(TEXT("UnchangingName_2"));
 	const FName SrcExportObjectFName = SrcAssetFName;
 	const FName DstExportObjectFName = DstAssetFName;
 	const FName SrcExportSubObjectFName(TEXT("SrcSubObject"));
 	const FName DstExportSubObjectFName(TEXT("DstSubObject"));
+	const FName SrcExportOnlyNumberChangesFName(TEXT("OnlyNumberChange_1"));
+	const FName DstExportOnlyNumberChangesFName(TEXT("OnlyNumberChange_2"));
 	const FName DummyImportPackagePathFName(TEXT("/DummyMount/DummyPackage"));
 	// Import FNames
 	const FName SrcEngineModuleImportObjectName(TEXT("/Script/SrcEngineModule"));
@@ -3716,6 +3726,7 @@ TEST_CASE_NAMED(FAssetHeaderPatcherTests, "AssetHeaderPatcher", "[AssetHeaderPat
 	TMap<FString, FString> PackageRenameMap =
 	{
 		{ SrcPackagePath, DstPackagePath },
+		{ SrcNumberChangedPath, DstNumberChangedPath},
 	};
 
 	auto MakeImport = [](const FName ObjectName, const FPackageIndex OuterIndex, const FName ClassPackage, const FName ClassName, const FName PackageName)
@@ -3913,6 +3924,14 @@ TEST_CASE_NAMED(FAssetHeaderPatcherTests, "AssetHeaderPatcher", "[AssetHeaderPat
 		{
 			.Src = MakeExport(SrcExportSubObjectFName, FPackageIndex::FromExport(1), FPackageIndex::FromExport(0), FPackageIndex::FromImport(1), FPackageIndex::FromImport(3), FPackageIndex::FromImport(5)),
 			.Dst = MakeExport(DstExportSubObjectFName, FPackageIndex::FromExport(1), FPackageIndex::FromExport(0), FPackageIndex::FromImport(1), FPackageIndex::FromImport(3), FPackageIndex::FromImport(5))
+		},
+		// SrcPackage.OnlyNumberChange_1												-> DstPackage.OnlyNumberChange_2
+		// Super: /Script/SrcEngineModule.SrcTypeA										-> /Script/DstEngineModule.DstTypeA
+		// Class: /Script/SrcEngineModule.SrcTypeA.MovedButNotRenamedProperty			-> /Script/DstEngineModule.DstTypeA.MovedButNotRenamedProperty
+		// Template: /Script/SrcEngineModule.OnlySubTypeChanged.SrcPropertyToChangeB	-> /Script/SrcEngineModule.OnlySubTypeChanged.DstPropertyToChangeB
+		{
+			.Src = MakeExport(SrcExportOnlyNumberChangesFName, FPackageIndex::FromExport(1), FPackageIndex::FromExport(0), FPackageIndex::FromImport(1), FPackageIndex::FromImport(3), FPackageIndex::FromImport(5)),
+			.Dst = MakeExport(DstExportOnlyNumberChangesFName, FPackageIndex::FromExport(1), FPackageIndex::FromExport(0), FPackageIndex::FromImport(1), FPackageIndex::FromImport(3), FPackageIndex::FromImport(5))
 		}
 	};
 
@@ -3946,6 +3965,7 @@ TEST_CASE_NAMED(FAssetHeaderPatcherTests, "AssetHeaderPatcher", "[AssetHeaderPat
 			
 			// NameTable
 			AddToNameTable(SrcPackagePathFName);
+			AddToNameTable(SrcNumberChangedPathFName);
 			AddToNameTable(SrcAssetFName);
 			AddToNameTable(DummyImportPackagePathFName);
 			// Redirected names (softpaths don't include the subpath in the name table
@@ -3954,6 +3974,7 @@ TEST_CASE_NAMED(FAssetHeaderPatcherTests, "AssetHeaderPatcher", "[AssetHeaderPat
 			// Export Table Names
 			AddToNameTable(SrcExportObjectFName);
 			AddToNameTable(SrcExportSubObjectFName);
+			AddToNameTable(SrcExportOnlyNumberChangesFName);
 			// Import Table Names
 			AddToNameTable(SrcEngineModuleImportObjectName);
 			AddToNameTable(SrcTypeAImportObjectName);
@@ -4062,6 +4083,13 @@ TEST_CASE_NAMED(FAssetHeaderPatcherTests, "AssetHeaderPatcher", "[AssetHeaderPat
 			CHECK(Patcher.DoPatch(Actual));
 			CHECK_EQUALS(TEXT("Generated MaterialFunctionInterface Editor Only Data mapping"), Actual, Expected);
 		}
+		
+		{
+			FString Actual(TEXT(R"(/SrcMount/SomePath/UnchangingName_1)"));
+			const FString Expected(TEXT(R"(/DstMount/SomePath/UnchangingName_2)"));
+			CHECK(Patcher.DoPatch(Actual));
+			CHECK_EQUALS(TEXT("Patch string with direct match"), Actual, Expected);
+		}
 
 		SECTION("Verse Mountpoints")
 		{
@@ -4124,6 +4152,13 @@ TEST_CASE_NAMED(FAssetHeaderPatcherTests, "AssetHeaderPatcher", "[AssetHeaderPat
 				CHECK_EQUALS(TEXT("Patch string with direct match"), Actual, Expected);
 			}
 
+			{
+				FString Actual(TEXT(R"(/SrcMount/SomePath/UnchangingName_1)"));
+				const FString Expected(TEXT(R"(/DstMount/SomePath/UnchangingName_2)"));
+				CHECK(Patcher.DoPatch(Actual));
+				CHECK_EQUALS(TEXT("Patch string with direct match"), Actual, Expected);
+			}
+
 			{
 				FString Actual(TEXT(R"(/SrcMount/SomePath/SrcPackage2)"));
 				const FString Expected = Actual; // Must be a copy
@@ -4149,6 +4184,13 @@ TEST_CASE_NAMED(FAssetHeaderPatcherTests, "AssetHeaderPatcher", "[AssetHeaderPat
 				CHECK_EQUALS(TEXT("Patch sub-object path"), Actual, Expected);
 			}
 
+			{
+				FString Actual(TEXT(R"(/SrcMount/SomePath/UnchangingName_1.UnchangingName_1:AnOuter.To.A.SubObject)"));
+				const FString Expected(TEXT(R"(/DstMount/SomePath/UnchangingName_2.UnchangingName_2:AnOuter.To.A.SubObject)"));
+				CHECK(Patcher.DoPatch(Actual));
+				CHECK_EQUALS(TEXT("Patch sub-object path"), Actual, Expected);
+			}
+
 			// Worth adding support for in the future, but at the moment we cannot patch various parts of unquoted 
 			// sub-object paths (that are specifically strings in the header, FNames are fine). In this case we 
 			// can't patch the package path because the top-level asset (UnmappedObject) has no mapping for patching
@@ -4192,6 +4234,34 @@ TEST_CASE_NAMED(FAssetHeaderPatcherTests, "AssetHeaderPatcher", "[AssetHeaderPat
 					CHECK_EQUALS(TEXT("Patch default blueprint generated class object path with quotes"), Actual, Expected);
 				}
 
+				{
+					FString Actual(TEXT(R"('/SrcMount/SomePath/UnchangingName_1')"));
+					const FString Expected(TEXT(R"('/DstMount/SomePath/UnchangingName_2')"));
+					CHECK(Patcher.DoPatch(Actual));
+					CHECK_EQUALS(TEXT("Patch package path with quotes"), Actual, Expected);
+				}
+
+				{
+					FString Actual(TEXT(R"('/SrcMount/SomePath/UnchangingName_1.UnchangingName_1')"));
+					const FString Expected(TEXT(R"('/DstMount/SomePath/UnchangingName_2.UnchangingName_2')"));
+					CHECK(Patcher.DoPatch(Actual));
+					CHECK_EQUALS(TEXT("Patch object path with quotes"), Actual, Expected);
+				}
+
+				{
+					FString Actual(TEXT(R"('/SrcMount/SomePath/UnchangingName_1.UnchangingName_1_C')"));
+					const FString Expected(TEXT(R"('/DstMount/SomePath/UnchangingName_2.UnchangingName_2_C')"));
+					CHECK(Patcher.DoPatch(Actual));
+					CHECK_EQUALS(TEXT("Patch blueprint generated class with quotes"), Actual, Expected);
+				}
+
+				{
+					FString Actual(TEXT(R"('/SrcMount/SomePath/UnchangingName_1.Default__UnchangingName_1_C')"));
+					const FString Expected(TEXT(R"('/DstMount/SomePath/UnchangingName_2.Default__UnchangingName_2_C')"));
+					CHECK(Patcher.DoPatch(Actual));
+					CHECK_EQUALS(TEXT("Patch default blueprint generated class object path with quotes"), Actual, Expected);
+				}
+
 				// Do not replace strings that happen to overlap with ObjectNames that are remapped
 				{
 					FString Actual(TEXT(R"('SrcPackage')"));
@@ -4231,6 +4301,34 @@ TEST_CASE_NAMED(FAssetHeaderPatcherTests, "AssetHeaderPatcher", "[AssetHeaderPat
 					CHECK_EQUALS(TEXT("Patch default blueprint generated class object path with quotes"), Actual, Expected);
 				}
 
+				{
+					FString Actual(TEXT(R"("/SrcMount/SomePath/UnchangingName_1")"));
+					const FString Expected(TEXT(R"("/DstMount/SomePath/UnchangingName_2")"));
+					CHECK(Patcher.DoPatch(Actual));
+					CHECK_EQUALS(TEXT("Patch package path with quotes"), Actual, Expected);
+				}
+
+				{
+					FString Actual(TEXT(R"("/SrcMount/SomePath/UnchangingName_1.UnchangingName_1")"));
+					const FString Expected(TEXT(R"("/DstMount/SomePath/UnchangingName_2.UnchangingName_2")"));
+					CHECK(Patcher.DoPatch(Actual));
+					CHECK_EQUALS(TEXT("Patch object path with quotes"), Actual, Expected);
+				}
+
+				{
+					FString Actual(TEXT(R"("/SrcMount/SomePath/UnchangingName_1.UnchangingName_1_C")"));
+					const FString Expected(TEXT(R"("/DstMount/SomePath/UnchangingName_2.UnchangingName_2_C")"));
+					CHECK(Patcher.DoPatch(Actual));
+					CHECK_EQUALS(TEXT("Patch blueprint generated class with quotes"), Actual, Expected);
+				}
+
+				{
+					FString Actual(TEXT(R"("/SrcMount/SomePath/UnchangingName_1.Default__UnchangingName_1_C")"));
+					const FString Expected(TEXT(R"("/DstMount/SomePath/UnchangingName_2.Default__UnchangingName_2_C")"));
+					CHECK(Patcher.DoPatch(Actual));
+					CHECK_EQUALS(TEXT("Patch default blueprint generated class object path with quotes"), Actual, Expected);
+				}
+
 				// Do not replace strings that happen to overlap with ObjectNames that are remapped
 				{
 					FString Actual(TEXT(R"("SrcPackage")"));
@@ -4255,6 +4353,19 @@ TEST_CASE_NAMED(FAssetHeaderPatcherTests, "AssetHeaderPatcher", "[AssetHeaderPat
 					CHECK_EQUALS(TEXT("Patch substring with quoted package, object and sub-object paths"), Actual, Expected);
 				}
 
+				{
+					FString Actual(
+						TEXT(R"(((ReferenceNodePath="/SrcMount/SomePath/UnchangingName_1.UnchangingName_1:RigVMModel.Setup Arm",)")
+						TEXT(R"(((Package="/SrcMount/SomePath/UnchangingName_1",)")
+						TEXT(R"(HostObject="/SrcMount/SomePath/UnchangingName_1.UnchangingName_1_C")))"));
+					FString Expected(
+						TEXT(R"(((ReferenceNodePath="/DstMount/SomePath/UnchangingName_2.UnchangingName_2:RigVMModel.Setup Arm",)")
+						TEXT(R"(((Package="/DstMount/SomePath/UnchangingName_2",)")
+						TEXT(R"(HostObject="/DstMount/SomePath/UnchangingName_2.UnchangingName_2_C")))"));
+					CHECK(Patcher.DoPatch(Actual));
+					CHECK_EQUALS(TEXT("Patch substring with quoted package, object and sub-object paths"), Actual, Expected);
+				}
+
 				// Do not replace strings that happen to overlap with ObjectNames that are remapped(i.e "SrcPackage=" is not transformed to "DstPackage=") 
 				{
 					FString Actual(
@@ -4374,6 +4485,39 @@ TEST_CASE_NAMED(FAssetHeaderPatcherTests, "AssetHeaderPatcher", "[AssetHeaderPat
 			CHECK_EQUALS(TEXT("SoftObject patching doesn't implicitly update the PackageFileSummary"), Patcher.Summary.NameCount, OriginalNameTableCount);
 		}
 
+		{
+			ResetPatcher();
+
+			FSoftObjectPath Actual(TEXT("/SrcMount/SomePath/UnchangingName_1.UnchangingName_1"));
+			FSoftObjectPath Expected(TEXT("/DstMount/SomePath/UnchangingName_2.UnchangingName_2"));
+			CHECK(Patcher.DoPatch(Actual));
+			CHECK_EQUALS(TEXT("SoftObjectPath patching"), Actual, Expected);
+			CHECK_EQUALS(TEXT("SoftObject patching doesn't implicitly update the NameTable"), Patcher.NameTable[1], SrcNumberChangedPathFName);
+			CHECK_EQUALS(TEXT("SoftObject patching doesn't implicitly update the PackageFileSummary"), Patcher.Summary.NameCount, OriginalNameTableCount);
+			Patcher.PatchNameTable();
+			// Numbered FNames have no number in the NameTable
+			CHECK_EQUALS(TEXT("SoftObject patching updates NameTable entry"), Patcher.NameTable[1].GetDisplayIndex(), DstNumberChangedPathFName.GetDisplayIndex());
+			CHECK_EQUALS(TEXT("SoftObject patching updates NameTable entry"), Patcher.NameTable[1].GetNumber(), 0);
+			CHECK_EQUALS(TEXT("SoftObject patching doesn't implicitly update the PackageFileSummary"), Patcher.Summary.NameCount, OriginalNameTableCount);
+		}
+
+		{
+			ResetPatcher();
+
+			FSoftObjectPath Actual(TEXT("/SrcMount/SomePath/UnchangingName_1.UnchangingName_1:Some.UnchangingName_1.Subobject"));
+			// Note we do not replace the sub-object "UnchangingName_1" despite it matching the original package and object name
+			FSoftObjectPath Expected(TEXT("/DstMount/SomePath/UnchangingName_2.UnchangingName_2:Some.UnchangingName_1.Subobject"));
+			CHECK(Patcher.DoPatch(Actual));
+			CHECK_EQUALS(TEXT("SoftObjectPath with sub-object path patching"), Actual, Expected);
+			CHECK_EQUALS(TEXT("SoftObject patching doesn't implicitly update the NameTable"), Patcher.NameTable[1], SrcNumberChangedPathFName);
+			CHECK_EQUALS(TEXT("SoftObject patching doesn't implicitly update the PackageFileSummary"), Patcher.Summary.NameCount, OriginalNameTableCount);
+			Patcher.PatchNameTable();
+			// Numbered FNames have no number in the NameTable
+			CHECK_EQUALS(TEXT("SoftObject patching updates NameTable entry"), Patcher.NameTable[1].GetDisplayIndex(), DstNumberChangedPathFName.GetDisplayIndex());
+			CHECK_EQUALS(TEXT("SoftObject patching updates NameTable entry"), Patcher.NameTable[1].GetNumber(), 0);
+			CHECK_EQUALS(TEXT("SoftObject patching doesn't implicitly update the PackageFileSummary"), Patcher.Summary.NameCount, OriginalNameTableCount);
+		}
+
 #if WITH_EDITOR
 		SECTION("DoPatch(FSoftObjectPath)-WithRedirector")
 		{
@@ -4451,6 +4595,38 @@ TEST_CASE_NAMED(FAssetHeaderPatcherTests, "AssetHeaderPatcher", "[AssetHeaderPat
 			CHECK_EQUALS(TEXT("TopLevelAssetPatch(string) patching updates NameTable entry"), Patcher.NameTable[0], DstPackagePathFName);
 			CHECK_EQUALS(TEXT("TopLevelAssetPatch(string) patching doesn't implicitly update the PackageFileSummary"), Patcher.Summary.NameCount, OriginalNameTableCount);
 		}
+
+		{
+			ResetPatcher();
+
+			FTopLevelAssetPath Actual(SrcNumberChangedPath, SrcNumberAssetFName);
+			FTopLevelAssetPath Expected(DstNumberChangedPath, DstNumberAssetFName);
+			CHECK(Patcher.DoPatch(Actual));
+			CHECK_EQUALS(TEXT("TopLevelAssetPatch(FName,FName) patching"), Actual, Expected);
+			CHECK_EQUALS(TEXT("TopLevelAssetPatch(FName,FName) patching doesn't implicitly update the NameTable"), Patcher.NameTable[1], SrcNumberChangedPathFName);
+			CHECK_EQUALS(TEXT("TopLevelAssetPatch(FName,FName) patching doesn't implicitly update the PackageFileSummary"), Patcher.Summary.NameCount, OriginalNameTableCount);
+			Patcher.PatchNameTable();
+			// Numbered FNames have no number in the NameTable
+			CHECK_EQUALS(TEXT("TopLevelAssetPatch(FName,FName) patching updates NameTable entry"), Patcher.NameTable[1].GetDisplayIndex(), DstNumberChangedPathFName.GetDisplayIndex());
+			CHECK_EQUALS(TEXT("TopLevelAssetPatch(FName,FName) patching updates NameTable entry"), Patcher.NameTable[1].GetNumber(), 0);
+			CHECK_EQUALS(TEXT("TopLevelAssetPatch(FName,FName) patching doesn't implicitly update the PackageFileSummary"), Patcher.Summary.NameCount, OriginalNameTableCount);
+		}
+
+		{
+			ResetPatcher();
+
+			FTopLevelAssetPath Actual(SrcNumberChangedObjectPath);
+			FTopLevelAssetPath Expected(DstNumberChangedObjectPath);
+			CHECK(Patcher.DoPatch(Actual));
+			CHECK_EQUALS(TEXT("TopLevelAssetPatch(string) patching"), Actual, Expected);
+			CHECK_EQUALS(TEXT("TopLevelAssetPatch(string) patching doesn't implicitly update the NameTable"), Patcher.NameTable[1], SrcNumberChangedPathFName);
+			CHECK_EQUALS(TEXT("TopLevelAssetPatch(string) patching doesn't implicitly update the PackageFileSummary"), Patcher.Summary.NameCount, OriginalNameTableCount);
+			Patcher.PatchNameTable();
+			// Numbered FNames have no number in the NameTable
+			CHECK_EQUALS(TEXT("TopLevelAssetPatch(string) patching updates NameTable entry"), Patcher.NameTable[1].GetDisplayIndex(), DstNumberChangedPathFName.GetDisplayIndex());
+			CHECK_EQUALS(TEXT("TopLevelAssetPatch(string) patching updates NameTable entry"), Patcher.NameTable[1].GetNumber(), 0);
+			CHECK_EQUALS(TEXT("TopLevelAssetPatch(string) patching doesn't implicitly update the PackageFileSummary"), Patcher.Summary.NameCount, OriginalNameTableCount);
+		}
 	}
 
 	SECTION("DoPatch(FGatherableTextData)")
@@ -4480,12 +4656,12 @@ TEST_CASE_NAMED(FAssetHeaderPatcherTests, "AssetHeaderPatcher", "[AssetHeaderPat
 			CHECK_EQUALS(TEXT("FGatherableTextData patching does update SourceData.SourceString[].SiteDescription"), Actual.SourceSiteContexts[0].SiteDescription, Expected.SourceSiteContexts[0].SiteDescription);
 			CHECK_EQUALS(TEXT("FGatherableTextData patching doesn't implicitly update the PackageFileSummary"), Patcher.Summary.NameCount, OriginalNameTableCount);
 			CHECK_EQUALS(TEXT("FGatherableTextData patching doesn't implicitly update the NameTable"), Patcher.NameTable[0], SrcPackagePathFName);
-			CHECK_EQUALS(TEXT("FGatherableTextData patching doesn't implicitly update the NameTable"), Patcher.NameTable[1], SrcAssetFName);
+			CHECK_EQUALS(TEXT("FGatherableTextData patching doesn't implicitly update the NameTable"), Patcher.NameTable[2], SrcAssetFName);
 			Patcher.PatchNameTable();
 			// FGatherableTexData doesn't contain FNames so we shouldn't have updated the NameTable at all
 			CHECK_EQUALS(TEXT("FGatherableTextData patching doesn't implicitly update the PackageFileSummary"), Patcher.Summary.NameCount, OriginalNameTableCount);
 			CHECK_EQUALS(TEXT("FGatherableTextData patching updates NameTable entry"), Patcher.NameTable[0], SrcPackagePathFName);
-			CHECK_EQUALS(TEXT("FGatherableTextData patching doesn't implicitly update the NameTable"), Patcher.NameTable[1], SrcAssetFName);
+			CHECK_EQUALS(TEXT("FGatherableTextData patching doesn't implicitly update the NameTable"), Patcher.NameTable[2], SrcAssetFName);
 		}
 	}
 
@@ -4565,6 +4741,7 @@ TEST_CASE_NAMED(FAssetHeaderPatcherTests, "AssetHeaderPatcher", "[AssetHeaderPat
 				{ ECoreRedirectFlags::Type_Package, TEXT("/Module/_Verse/VerseAsset"),                                                 TEXT("/Module/_Verse") },
 				{ ECoreRedirectFlags::Type_Object,  TEXT("/Module/_Verse/VerseAsset.some_verse_class"),                                TEXT("/Module/_Verse.VerseAsset-some_verse_class") },
 				{ ECoreRedirectFlags::Type_Object,  TEXT("/SrcMount/SomePath/SrcPackage.SrcPackage:SrcSubObject"),                     TEXT("/DstMount/SomePath/DstPackage.DstPackage:DstSubObject")},
+				{ ECoreRedirectFlags::Type_Object,  TEXT("/SrcMount/SomePath/SrcPackage.SrcPackage:OnlyNumberChange_1"),               TEXT("/SrcMount/SomePath/SrcPackage.SrcPackage:OnlyNumberChange_2")},
 			};
 			CHECK(FCoreRedirects::AddRedirectList(ImportTableRedirects, TEXT("Asset Header Patcher Tests - FObjectResource Patching")));
 
