[BUGFIX] Importing Single File without interchange will lead to crash in SlowTask
- Removed setting the completed work on
slow task preemtively

#rnx

diff --git a/Engine/Source/Developer/AssetTools/Private/AssetTools.cpp b/Engine/Source/Developer/AssetTools/Private/AssetTools.cpp
--- a/Engine/Source/Developer/AssetTools/Private/AssetTools.cpp
+++ b/Engine/Source/Developer/AssetTools/Private/AssetTools.cpp
@@ -3544,11 +3544,7 @@ TArray<UObject*> UAssetToolsImpl::ImportAssetsInternal(const TArray<FString>& Fi
 		//If we're importing a single file, then the factory policy will dictate if the import if cancelable.
 		SlowTask.MakeDialog(true);
 	}
-	else
-	{
-		SlowTask.CompletedWork = static_cast<float>(FilesAndDestinations.Num());
-	}
-
+	
 	struct FInterchangeImportStatus
 	{
 		explicit FInterchangeImportStatus(int32 NumFiles)
