Asynchronous content browser asset creation.

Swarm review: 38667873

diff --git a/Engine/Plugins/Editor/ContentBrowser/ContentBrowserAssetDataSource/Source/ContentBrowserAssetDataSource/Private/ContentBrowserAssetDataSource.cpp b/Engine/Plugins/Editor/ContentBrowser/ContentBrowserAssetDataSource/Source/ContentBrowserAssetDataSource/Private/ContentBrowserAssetDataSource.cpp
--- a/Engine/Plugins/Editor/ContentBrowser/ContentBrowserAssetDataSource/Source/ContentBrowserAssetDataSource/Private/ContentBrowserAssetDataSource.cpp
+++ b/Engine/Plugins/Editor/ContentBrowser/ContentBrowserAssetDataSource/Source/ContentBrowserAssetDataSource/Private/ContentBrowserAssetDataSource.cpp
@@ -3937,12 +3937,26 @@ void UContentBrowserAssetDataSource::OnBeginCreateAsset(const FName InDefaultAss
 	FContentBrowserModule& ContentBrowserModule = FModuleManager::Get().GetModuleChecked<FContentBrowserModule>(TEXT("ContentBrowser"));
 	FContentBrowserItemPath AssetPathToUse = ContentBrowserModule.Get().GetInitialPathToSaveAsset(FContentBrowserItemPath(InPackagePath, EContentBrowserPathType::Internal));
 
+	const bool bUseAsyncCreate = InFactory && EnumHasAnyFlags(InFactory->GetSupportedWorkflows(), EFactoryCreateWorkflow::Asynchronous);
 	const bool bShowDialogToPickPath = !AssetPathToUse.HasInternalPath() || (AssetPathToUse.GetInternalPathName() != InPackagePath);
 	if (bShowDialogToPickPath)
 	{
 		FAssetToolsModule& AssetToolsModule = FModuleManager::Get().LoadModuleChecked<FAssetToolsModule>("AssetTools");
 		const FString InitialInternalPath = AssetPathToUse.HasInternalPath() ? AssetPathToUse.GetInternalPathString() : TEXT("/Game");
-		AssetToolsModule.Get().CreateAssetWithDialog(InDefaultAssetName.ToString(), InitialInternalPath, ClassToUse, InFactory, NAME_None, /*bCallConfigureProperties*/ false);
+
+		if (bUseAsyncCreate)
+		{
+			AssetToolsModule.Get().CreateAssetWithDialogAsync(InDefaultAssetName.ToString(), InitialInternalPath, ClassToUse, InFactory, FAssetCreateComplete(), FAssetCreateCancelled(), NAME_None);
+		}
+		else
+		{
+			AssetToolsModule.Get().CreateAssetWithDialog(InDefaultAssetName.ToString(), InitialInternalPath, ClassToUse, InFactory, NAME_None, /*bCallConfigureProperties*/ false);
+		}
+	}
+	else if (bUseAsyncCreate)
+	{
+		FAssetToolsModule& AssetToolsModule = FModuleManager::Get().LoadModuleChecked<FAssetToolsModule>("AssetTools");
+		AssetToolsModule.Get().CreateAssetAsync(InDefaultAssetName.ToString(), InPackagePath.ToString(), ClassToUse, InFactory);
 	}
 	else
 	{
diff --git a/Engine/Source/Developer/AssetTools/Private/AssetTools.cpp b/Engine/Source/Developer/AssetTools/Private/AssetTools.cpp
--- a/Engine/Source/Developer/AssetTools/Private/AssetTools.cpp
+++ b/Engine/Source/Developer/AssetTools/Private/AssetTools.cpp
@@ -1726,6 +1726,13 @@ UObject* UAssetToolsImpl::CreateAsset(const FString& AssetName, const FString& P
 		return nullptr;
 	}
 
+	// This is the class default and this should not fail for all current factories.
+	if (Factory && !EnumHasAnyFlags(Factory->GetSupportedWorkflows(), EFactoryCreateWorkflow::Default))
+	{
+		FMessageDialog::Open(EAppMsgType::Ok, UE::Private::GetInvalidFactoryMessage());
+		return nullptr;
+	}
+
 	// Verify the factory supports the asset class
 	if (AssetClass && Factory && 
 		((Factory->SupportedClass != nullptr && !ensure(AssetClass->IsChildOf(Factory->GetSupportedClass()))) || 
@@ -1774,6 +1781,58 @@ UObject* UAssetToolsImpl::CreateAsset(const FString& AssetName, const FString& P
 	return NewObj;
 }
 
+bool UAssetToolsImpl::CreateAssetAsync(const FString& AssetName, const FString& PackagePath, UClass* AssetClass, UFactory* Factory,
+	FAssetCreateComplete OnComplete, FAssetCreateCancelled OnCancelled, FName CallingContext)
+{
+	// Verify the factory class
+	if (!ensure(AssetClass && Factory))
+	{
+		FMessageDialog::Open(EAppMsgType::Ok, LOCTEXT("MustSupplyClassOrFactory", "The new asset wasn't created due to a problem finding the appropriate factory or class for the new asset."));
+		OnCancelled.ExecuteIfBound(Factory);
+		return false;
+	}
+
+	if (!EnumHasAnyFlags(Factory->GetSupportedWorkflows(), EFactoryCreateWorkflow::Asynchronous))
+	{
+		FMessageDialog::Open(EAppMsgType::Ok, UE::Private::GetInvalidFactoryMessage());
+		OnCancelled.ExecuteIfBound(Factory);
+		return false;
+	}
+
+	// Verify the factory supports the asset class
+	if (((Factory->SupportedClass != nullptr && !ensure(AssetClass->IsChildOf(Factory->GetSupportedClass()))) ||
+			(Factory->SupportedClass == nullptr && !ensure(Factory->DoesSupportClass(AssetClass)))))
+	{
+		FMessageDialog::Open(EAppMsgType::Ok, UE::Private::GetInvalidFactoryMessage());
+		OnCancelled.ExecuteIfBound(Factory);
+		return false;
+	}
+
+	const FString PackageName = UPackageTools::SanitizePackageName(PackagePath + TEXT("/") + AssetName);
+
+	// Make sure we can create the asset without conflicts
+	if (!CanCreateAsset(AssetName, PackageName, LOCTEXT("CreateANewObject", "Create a new object")))
+	{
+		OnCancelled.ExecuteIfBound(Factory);
+		return false;
+	}
+
+	UClass* ClassToUse = AssetClass ? AssetClass : (Factory ? Factory->GetSupportedClass() : nullptr);
+
+	UPackage* Pkg = CreatePackage(*PackageName);
+	EObjectFlags Flags = RF_Public | RF_Standalone | RF_Transactional;
+
+	FEditorDelegates::OnConfigureNewAssetProperties.Broadcast(Factory);
+
+	TStrongObjectPtr<UPackage> DontGCPackage(Pkg);
+
+	return Factory->ConfigurePropertiesAsync(
+		FFactoryConfigurePropertiesComplete::CreateUObject(this, &UAssetToolsImpl::CreateAssetAsync_OnComplete, DontGCPackage, ClassToUse, 
+			AssetName, Flags, CallingContext, PackageName, OnComplete, OnCancelled),
+		FFactoryConfigurePropertiesCancelled()
+	);
+}
+
 void UAssetToolsImpl::CreateAssetsFrom(TConstArrayView<UObject*> SourceObjects, UClass* CreateAssetType, const FString& DefaultSuffix, TFunctionRef<UFactory*(UObject*)> FactoryConstructor, FName CallingContext)
 {
 	if ( SourceObjects.Num() == 1 )
@@ -1836,6 +1895,33 @@ void UAssetToolsImpl::CreateAssetsFrom(TConstArrayView<UObject*> SourceObjects,
 	}
 }
 
+void UAssetToolsImpl::CreateAssetAsync_OnComplete(UFactory* Factory, TStrongObjectPtr<UPackage> Package, UClass* AssetClass, FString AssetName,
+	EObjectFlags Flags, FName CallingContext, FString PackageName, FAssetCreateComplete OnComplete, FAssetCreateCancelled OnCancelled)
+{
+	UPackage* Pkg = Package.Get();
+	UObject* NewObj = Factory->FactoryCreateNew(AssetClass, Pkg, FName(*AssetName), Flags, nullptr, GWarn, CallingContext);
+
+	if (NewObj)
+	{
+		Pkg->SetIsExternallyReferenceable(ShouldCreateAssetsAsExternallyReferenceableForPath(PackageName));
+
+		// Notify the asset registry
+		FAssetRegistryModule::AssetCreated(NewObj);
+
+		// analytics create record
+		UAssetToolsImpl::OnNewCreateRecord(AssetClass, false);
+
+		// Mark the package dirty...
+		Pkg->MarkPackageDirty();
+
+		OnComplete.ExecuteIfBound(Factory, NewObj);
+	}
+	else
+	{
+		OnCancelled.ExecuteIfBound(Factory);
+	}
+}
+
 UObject* UAssetToolsImpl::CreateAssetWithDialog(UClass* AssetClass, UFactory* Factory, FName CallingContext)
 {
 	if (Factory != nullptr)
@@ -1868,6 +1954,7 @@ UObject* UAssetToolsImpl::CreateAssetWithDialog(UClass* AssetClass, UFactory* Fa
 UObject* UAssetToolsImpl::CreateAssetWithDialog(const FString& AssetName, const FString& PackagePath, UClass* AssetClass, UFactory* Factory, FName CallingContext, const bool bCallConfigureProperties)
 {
 	FGCObjectScopeGuard DontGCFactory(Factory);
+
 	if(Factory)
 	{
 		FSaveAssetDialogConfig SaveAssetDialogConfig;
@@ -1887,11 +1974,11 @@ UObject* UAssetToolsImpl::CreateAssetWithDialog(const FString& AssetName, const
 		{
 			FMessageDialog::Open(EAppMsgType::Ok, UE::Private::GetInvalidFactoryMessage());
 			return nullptr;
-		}
-		
+		}		
 
 		FContentBrowserModule& ContentBrowserModule = FModuleManager::LoadModuleChecked<FContentBrowserModule>("ContentBrowser");
 		FString SaveObjectPath = ContentBrowserModule.Get().CreateModalSaveAssetDialog(SaveAssetDialogConfig);
+
 		if (!SaveObjectPath.IsEmpty())
 		{
 			bool bCreateAsset = true;
@@ -1916,6 +2003,105 @@ UObject* UAssetToolsImpl::CreateAssetWithDialog(const FString& AssetName, const
 	return nullptr;
 }
 
+bool UAssetToolsImpl::CreateAssetWithDialogAsync(UClass* AssetClass, UFactory* Factory, FAssetCreateComplete OnComplete,
+	FAssetCreateCancelled OnCancelled, FName CallingContext)
+{
+	if (!Factory)
+	{
+		OnCancelled.ExecuteIfBound(Factory);
+		return false;
+	}
+
+	// Determine the starting path. Try to use the most recently used directory
+	FString AssetPath;
+
+	const FString DefaultFilesystemDirectory = FEditorDirectories::Get().GetLastDirectory(ELastDirectory::NEW_ASSET);
+	if (DefaultFilesystemDirectory.IsEmpty() || !FPackageName::TryConvertFilenameToLongPackageName(DefaultFilesystemDirectory, AssetPath))
+	{
+		// No saved path, just use the game content root
+		AssetPath = TEXT("/Game");
+	}
+
+	FString PackageName;
+	FString AssetName;
+	CreateUniqueAssetName(AssetPath / Factory->GetDefaultNewAssetName(), TEXT(""), PackageName, AssetName);
+
+	return CreateAssetWithDialogAsync(AssetName, AssetPath, AssetClass, Factory, OnComplete, OnCancelled, CallingContext);
+}
+
+bool UAssetToolsImpl::CreateAssetWithDialogAsync(const FString& AssetName, const FString& PackagePath, UClass* AssetClass, UFactory* Factory,
+	FAssetCreateComplete OnComplete, FAssetCreateCancelled OnCancelled, FName CallingContext)
+{
+	if (!Factory)
+	{
+		OnCancelled.ExecuteIfBound(Factory);
+		return false;
+	}
+
+	if (!EnumHasAnyFlags(Factory->GetSupportedWorkflows(), EFactoryCreateWorkflow::Asynchronous))
+	{
+		FMessageDialog::Open(EAppMsgType::Ok, UE::Private::GetInvalidFactoryMessage());
+		OnCancelled.ExecuteIfBound(Factory);
+		return false;
+	}
+
+	FSaveAssetDialogConfig SaveAssetDialogConfig;
+	SaveAssetDialogConfig.DialogTitleOverride = LOCTEXT("SaveAssetDialogTitle", "Save Asset As");
+	SaveAssetDialogConfig.DefaultPath = PackagePath;
+	SaveAssetDialogConfig.DefaultAssetName = AssetName;
+	SaveAssetDialogConfig.ExistingAssetPolicy = ESaveAssetDialogExistingAssetPolicy::AllowButWarn;
+
+	if (Factory->GetSupportedClass() != nullptr)
+	{
+		SaveAssetDialogConfig.AssetClassNames.Add(Factory->GetSupportedClass()->GetClassPathName());
+	}
+	else if (AssetClass && Factory->DoesSupportClass(AssetClass))
+	{
+		SaveAssetDialogConfig.AssetClassNames.Add(AssetClass->GetClassPathName());
+	}
+	else
+	{
+		FMessageDialog::Open(EAppMsgType::Ok, UE::Private::GetInvalidFactoryMessage());
+		OnCancelled.ExecuteIfBound(Factory);
+		return false;
+	}
+
+	FContentBrowserModule& ContentBrowserModule = FModuleManager::LoadModuleChecked<FContentBrowserModule>("ContentBrowser");
+	FString SaveObjectPath = ContentBrowserModule.Get().CreateModalSaveAssetDialog(SaveAssetDialogConfig);
+
+	if (SaveObjectPath.IsEmpty())
+	{
+		OnCancelled.ExecuteIfBound(Factory);
+		return false;
+	}
+
+	FEditorDelegates::OnConfigureNewAssetProperties.Broadcast(Factory);
+
+	return Factory->ConfigurePropertiesAsync(
+		FFactoryConfigurePropertiesComplete::CreateUObject(this, &UAssetToolsImpl::CreateAssetWithDialogAsync_OnComplete,
+			AssetClass, CallingContext, SaveObjectPath, OnComplete, OnCancelled),
+		FFactoryConfigurePropertiesCancelled()
+	);
+}
+
+void UAssetToolsImpl::CreateAssetWithDialogAsync_OnComplete(UFactory* Factory, UClass* AssetClass, FName CallingContext, FString SaveObjectPath,
+	FAssetCreateComplete OnComplete, FAssetCreateCancelled OnCancelled)
+{
+	const FString SavePackageName = FPackageName::ObjectPathToPackageName(SaveObjectPath);
+	const FString SavePackagePath = FPaths::GetPath(SavePackageName);
+	const FString SaveAssetName = FPaths::GetBaseFilename(SavePackageName);
+	FEditorDirectories::Get().SetLastDirectory(ELastDirectory::NEW_ASSET, SavePackagePath);
+
+	if (UObject* NewObj = CreateAsset(SaveAssetName, SavePackagePath, AssetClass, Factory, CallingContext))
+	{
+		OnComplete.ExecuteIfBound(Factory, NewObj);
+	}
+	else
+	{
+		OnCancelled.ExecuteIfBound(Factory);
+	}
+}
+
 UObject* UAssetToolsImpl::DuplicateAssetWithDialog(const FString& AssetName, const FString& PackagePath, UObject* OriginalObject)
 {
 	return DuplicateAssetWithDialogAndTitle(AssetName, PackagePath, OriginalObject, LOCTEXT("DuplicateAssetDialogTitle", "Duplicate Asset As"));
diff --git a/Engine/Source/Developer/AssetTools/Private/AssetTools.h b/Engine/Source/Developer/AssetTools/Private/AssetTools.h
--- a/Engine/Source/Developer/AssetTools/Private/AssetTools.h
+++ b/Engine/Source/Developer/AssetTools/Private/AssetTools.h
@@ -60,7 +60,8 @@ PRAGMA_DISABLE_DEPRECATION_WARNINGS
 UCLASS(transient)
 class UAssetToolsImpl : public UObject, public IAssetTools
 {
-	GENERATED_BODY() 
+	GENERATED_BODY()
+
 public:
 	UAssetToolsImpl(const FObjectInitializer& ObjectInitializer);
 
@@ -84,9 +85,12 @@ public:
 	virtual void GetClassTypeActionsList( TArray<TWeakPtr<IClassTypeActions>>& OutClassTypeActionsList ) const override;
 	virtual TWeakPtr<IClassTypeActions> GetClassTypeActionsForClass( UClass* Class ) const override;
 	virtual UObject* CreateAsset(const FString& AssetName, const FString& PackagePath, UClass* AssetClass, UFactory* Factory, FName CallingContext = NAME_None) override;
-	virtual void CreateAssetsFrom(TConstArrayView<UObject*> SourceObjects, UClass* CreateAssetType, const FString& DefaultSuffix, TFunctionRef<UFactory*(UObject*)> FactoryConstructor, FName CallingContext = NAME_None);
+	virtual bool CreateAssetAsync(const FString& AssetName, const FString& PackagePath, UClass* AssetClass, UFactory* Factory, FAssetCreateComplete OnComplete = FAssetCreateComplete(), FAssetCreateCancelled OnCancelled = FAssetCreateCancelled(), FName CallingContext = NAME_None) override;
+	virtual void CreateAssetsFrom(TConstArrayView<UObject*> SourceObjects, UClass* CreateAssetType, const FString& DefaultSuffix, TFunctionRef<UFactory*(UObject*)> FactoryConstructor, FName CallingContext) override;
 	virtual UObject* CreateAssetWithDialog(UClass* AssetClass, UFactory* Factory, FName CallingContext = NAME_None) override;
-	virtual UObject* CreateAssetWithDialog(const FString& AssetName, const FString& PackagePath, UClass* AssetClass, UFactory* Factory, FName CallingContext = NAME_None, const bool bCallConfigureProperties = true) override;
+	virtual UObject* CreateAssetWithDialog(const FString& AssetName, const FString& PackagePath, UClass* AssetClass, UFactory* Factory, FName CallingContext, const bool bCallConfigureProperties = true) override;
+	virtual bool CreateAssetWithDialogAsync(const FString& AssetName, const FString& PackagePath, UClass* AssetClass, UFactory* Factory, FAssetCreateComplete OnComplete, FAssetCreateCancelled OnCancelled, FName CallingContext) override;
+	virtual bool CreateAssetWithDialogAsync(UClass* AssetClass, UFactory* Factory, FAssetCreateComplete OnComplete = FAssetCreateComplete(), FAssetCreateCancelled OnCancelled = FAssetCreateCancelled(), FName CallingContext = NAME_None) override;
 	virtual UObject* DuplicateAsset(const FString& AssetName, const FString& PackagePath, UObject* OriginalObject) override;
 	virtual UObject* DuplicateAssetWithDialog(const FString& AssetName, const FString& PackagePath, UObject* OriginalObject) override;
 	virtual UObject* DuplicateAssetWithDialogAndTitle(const FString& AssetName, const FString& PackagePath, UObject* OriginalObject, FText DialogTitle) override;
@@ -184,6 +188,7 @@ public:
 
 	virtual void SyncBrowserToAssets(const TArray<UObject*>& AssetsToSync) override;
 	virtual void SyncBrowserToAssets(const TArray<FAssetData>& AssetsToSync) override;
+
 public:
 	/** Gets the asset tools singleton as a FAssetTools for asset tools module use */
 	static UAssetToolsImpl& Get();
@@ -263,6 +268,14 @@ private:
 	/** Helper to remove an entry AssetTypeActionsList given a class returned by GetSupportedClass */
 	void RemoveAssetTypeActionBySupportedClass(const UClass* SupportedClass);
 
+	/** Called when async asset creation is completed successfully. */
+	void CreateAssetAsync_OnComplete(UFactory* Factory, TStrongObjectPtr<UPackage> Package, UClass* AssetClass, FString AssetName, EObjectFlags Flags,
+		FName CallingContext, FString PackageName, FAssetCreateComplete OnComplete, FAssetCreateCancelled OnCancelled);
+
+	/** Called when async asset creation with a file location picker dialog is completed successfully. */
+	void CreateAssetWithDialogAsync_OnComplete(UFactory* Factory, UClass* AssetClass, FName CallingContext, FString SaveObjectPath,
+		FAssetCreateComplete OnComplete, FAssetCreateCancelled OnCancelled);
+
 private:
 	uint64 CachedVersion = 0;
 	
diff --git a/Engine/Source/Developer/AssetTools/Public/IAssetTools.h b/Engine/Source/Developer/AssetTools/Public/IAssetTools.h
--- a/Engine/Source/Developer/AssetTools/Public/IAssetTools.h
+++ b/Engine/Source/Developer/AssetTools/Public/IAssetTools.h
@@ -275,6 +275,11 @@ class UAssetTools : public UInterface
 	GENERATED_UINTERFACE_BODY()
 };
 
+DECLARE_DYNAMIC_DELEGATE_TwoParams(FAssetCreateCompleteDynamic, UFactory*, Factory, UObject*, NewObject);
+DECLARE_DYNAMIC_DELEGATE_OneParam(FAssetCreateCancelledDynamic, UFactory*, Factory);
+DECLARE_DELEGATE_TwoParams(FAssetCreateComplete, UFactory* /* Factory */, UObject* /* NewObject */);
+DECLARE_DELEGATE_OneParam(FAssetCreateCancelled, UFactory* /* Factory */);
+
 class IAssetTools
 {
 	GENERATED_IINTERFACE_BODY()
@@ -342,6 +347,34 @@ public:
 	UFUNCTION(BlueprintCallable, Category = "Editor Scripting | Asset Tools")
 	virtual UObject* CreateAsset(const FString& AssetName, const FString& PackagePath, UClass* AssetClass, UFactory* Factory, FName CallingContext = NAME_None) = 0;
 
+	/**
+	 * Creates an asset with the specified name, path, and factory asynchronously
+	 *
+	 * @param AssetName the name of the new asset
+	 * @param PackagePath the package that will contain the new asset
+	 * @param AssetClass the class of the new asset
+	 * @param Factory the factory that will build the new asset
+	 * @param OnComplete called when the factory has successfully created the asset
+	 * @param OnCancelled called if the asynchronous method is cancelled
+	 * @param CallingContext optional name of the module or method calling CreateAsset() - this is passed to the factory
+	 * @return the new asset or NULL if it fails
+	 */
+	UFUNCTION(BlueprintCallable, Category = "Editor Scripting | Asset Tools")
+	virtual bool CreateAssetAsync(const FString& AssetName, const FString& PackagePath, UClass* AssetClass, UFactory* Factory, FAssetCreateCompleteDynamic OnComplete, FAssetCreateCancelledDynamic OnCancelled, FName CallingContext = NAME_None)
+	{
+		return CreateAssetAsync(
+			AssetName, 
+			PackagePath, 
+			AssetClass, 
+			Factory,
+			FAssetCreateComplete::CreateLambda([OnComplete](UFactory* Factory, UObject* NewObject) { OnComplete.ExecuteIfBound(Factory, NewObject); }),
+			FAssetCreateCancelled::CreateLambda([OnCancelled](UFactory* Factory) { OnCancelled.ExecuteIfBound(Factory); }),
+			CallingContext
+		);
+	}
+
+	virtual bool CreateAssetAsync(const FString& AssetName, const FString& PackagePath, UClass* AssetClass, UFactory* Factory, FAssetCreateComplete OnComplete = FAssetCreateComplete(), FAssetCreateCancelled OnCancelled = FAssetCreateCancelled(), FName CallingContext = NAME_None) = 0;
+
 	/**
 	 * Creates one or more assets using the source objects as the basis for the next type.  This is a common enough operation
 	 * it needed a utility.  In the case that you only have a single SourceObject, we'll lean on the content browser
@@ -373,6 +406,32 @@ public:
 	/** Opens an asset picker dialog and creates an asset with the path chosen in the dialog */
 	virtual UObject* CreateAssetWithDialog(UClass* AssetClass, UFactory* Factory, FName CallingContext = NAME_None) = 0;
 
+	/**
+	 * Opens an async dialog to create an asset.
+	 * @returns true if the process was started successfully.
+	 */
+	UFUNCTION(BlueprintCallable, Category = "Editor Scripting | Asset Tools")
+	virtual bool CreateAssetWithDialogAsync(const FString& AssetName, const FString& PackagePath, UClass* AssetClass, UFactory* Factory, FAssetCreateCompleteDynamic OnComplete, FAssetCreateCancelledDynamic OnCancelled, FName CallingContext = NAME_None)
+	{
+		return CreateAssetWithDialogAsync(
+			AssetName, 
+			PackagePath, 
+			AssetClass, 
+			Factory,
+			FAssetCreateComplete::CreateLambda([OnComplete](UFactory* Factory, UObject* NewObject) { OnComplete.ExecuteIfBound(Factory, NewObject); }),
+			FAssetCreateCancelled::CreateLambda([OnCancelled](UFactory* Factory) { OnCancelled.ExecuteIfBound(Factory); }),
+			CallingContext
+		);
+	}
+
+	virtual bool CreateAssetWithDialogAsync(const FString& AssetName, const FString& PackagePath, UClass* AssetClass, UFactory* Factory, FAssetCreateComplete OnComplete = FAssetCreateComplete(), FAssetCreateCancelled OnCancelled = FAssetCreateCancelled(), FName CallingContext = NAME_None) = 0;
+
+	/**
+	 * Opens an async dialog to create an asset.
+	 * @returns true if the process was started successfully.
+	 */
+	virtual bool CreateAssetWithDialogAsync(UClass* AssetClass, UFactory* Factory, FAssetCreateComplete OnComplete = FAssetCreateComplete(), FAssetCreateCancelled OnCancelled = FAssetCreateCancelled(), FName CallingContext = NAME_None) = 0;
+
 	/** Opens an asset picker dialog and creates an asset with the specified name and path. Uses OriginalObject as the duplication source. */
 	UFUNCTION(BlueprintCallable, Category = "Editor Scripting | Asset Tools")
 	virtual UObject* DuplicateAssetWithDialog(const FString& AssetName, const FString& PackagePath, UObject* OriginalObject) = 0;
diff --git a/Engine/Source/Editor/PropertyEditor/Private/UserInterface/PropertyEditor/SPropertyMenuAssetPicker.cpp b/Engine/Source/Editor/PropertyEditor/Private/UserInterface/PropertyEditor/SPropertyMenuAssetPicker.cpp
--- a/Engine/Source/Editor/PropertyEditor/Private/UserInterface/PropertyEditor/SPropertyMenuAssetPicker.cpp
+++ b/Engine/Source/Editor/PropertyEditor/Private/UserInterface/PropertyEditor/SPropertyMenuAssetPicker.cpp
@@ -352,18 +352,36 @@ void SPropertyMenuAssetPicker::SetValue( const FAssetData& AssetData )
 
 void SPropertyMenuAssetPicker::OnCreateNewAssetSelected(TWeakObjectPtr<UFactory> FactoryPtr)
 {
-	if (FactoryPtr.IsValid())
+	if (!FactoryPtr.IsValid())
+	{
+		return;
+	}
+
+	UFactory* FactoryInstance = DuplicateObject<UFactory>(FactoryPtr.Get(), GetTransientPackage());
+	// Ensure this object is not GC for the duration of CreateAssetWithDialog
+	TStrongObjectPtr<UFactory> FactoryStrong(FactoryInstance);
+	FAssetToolsModule& AssetToolsModule = FAssetToolsModule::GetModule();
+
+	if (EnumHasAnyFlags(FactoryInstance->GetSupportedWorkflows(), EFactoryCreateWorkflow::Asynchronous))
+	{
+		AssetToolsModule.Get().CreateAssetWithDialogAsync(
+			FactoryInstance->GetSupportedClass(), 
+			FactoryInstance,
+			FAssetCreateComplete::CreateSPLambda(
+				this,
+				[this](UFactory* Factory, UObject* InObject)
+				{
+					SetValue(InObject);
+				})
+		);
+	}
+	else if (EnumHasAnyFlags(FactoryInstance->GetSupportedWorkflows(), EFactoryCreateWorkflow::Default))
 	{
-		UFactory* FactoryInstance = DuplicateObject<UFactory>(FactoryPtr.Get(), GetTransientPackage());
-		// Ensure this object is not GC for the duration of CreateAssetWithDialog
-		FactoryInstance->AddToRoot();
-		FAssetToolsModule& AssetToolsModule = FAssetToolsModule::GetModule();
 		UObject* NewAsset = AssetToolsModule.Get().CreateAssetWithDialog(FactoryInstance->GetSupportedClass(), FactoryInstance);
 		if (NewAsset != nullptr)
 		{
 			SetValue(NewAsset);
 		}
-		FactoryInstance->RemoveFromRoot();
 	}
 }
 
diff --git a/Engine/Source/Editor/UnrealEd/Classes/Factories/Factory.h b/Engine/Source/Editor/UnrealEd/Classes/Factories/Factory.h
--- a/Engine/Source/Editor/UnrealEd/Classes/Factories/Factory.h
+++ b/Engine/Source/Editor/UnrealEd/Classes/Factories/Factory.h
@@ -13,6 +13,29 @@
 class UAssetImportTask;
 class FBulkData;
 
+/**
+ * Determines which workflows are available for a factory.
+ * All factories *should* support the default workflow!
+ */
+UENUM(meta = (Bitflags, UseEnumValuesAsMaskValuesInEditor = "true"))
+enum class EFactoryCreateWorkflow : uint8
+{
+	/**
+	 * Synchronous asset creation (potentially with a modal dialog).
+	 */
+	Default = 1 << 0,
+
+	/**
+	 * Asynchronous asset creation (usually with a dialog).
+	 * If supported, the UFactory::ConfigureProperties will not be called, but UFactory::ConfigurePropertiesAsync.
+	 */
+	Asynchronous = 1 << 1
+};
+ENUM_CLASS_FLAGS(EFactoryCreateWorkflow);
+
+DECLARE_DELEGATE_OneParam(FFactoryConfigurePropertiesComplete, UFactory* /* this */);
+DECLARE_DELEGATE_OneParam(FFactoryConfigurePropertiesCancelled, UFactory* /* this */);
+
 /**
  * Base class for all factories
  * An object responsible for creating and importing new objects.
@@ -177,6 +200,20 @@ public:
 		return true;
 	}
 
+	/**
+	 * Called instead of ConfigureProperties() when an asset is being created asynchronously.
+	 * It is not the responsibility of the caller of this function to stop the factory from being GC'd. It is the
+	 * responsibility of the implementor.
+	 * @param OnComplete  Called when the asynchronous process completes successfully.
+	 * @param OnCancelled Called when the asynchronous process was cancelled.
+	 * @return true if the async process was successfully started.
+	 */
+	virtual bool ConfigurePropertiesAsync(FFactoryConfigurePropertiesComplete OnComplete, FFactoryConfigurePropertiesCancelled OnCancelled)
+	{
+		OnCancelled.Execute(this);
+		return false;
+	}
+
 	// @todo document
 	UNREALED_API virtual bool ImportUntypedBulkDataFromText(const TCHAR*& Buffer, FBulkData& BulkData);
 
@@ -230,6 +267,14 @@ public:
 	 * @return the supported factory formats list.
 	 */
 	virtual TArray<FString> GetFormats() const { return Formats; }
+
+	/**
+	 * @return the supported factory workflows.
+	 */
+	virtual EFactoryCreateWorkflow GetSupportedWorkflows() const
+	{
+		return static_cast<EFactoryCreateWorkflow>(SupportedWorkflows);
+	}
 public:
 
 	/**
@@ -404,6 +449,10 @@ public:
 	UPROPERTY(BlueprintReadWrite, Category=Misc)
 	TObjectPtr<class UAssetImportTask> AssetImportTask;
 
+	/** Which workflows are supported. */
+	UPROPERTY(BlueprintReadWrite, Category=Misc, meta = (Bitmask))
+	uint8 SupportedWorkflows = static_cast<uint8>(EFactoryCreateWorkflow::Default);
+
 protected:
 
 	/** Name of the file currently being imported. */
