[DuplicatePluginCommandlet | AssetHeaderPatching]

- Extends the AssetHeaderPatcher to only patch out a subset of files passed to it
- Changes the DuplicatePluginCommandlet to keep track of all possible Src->Dst patching remappings possible and pass that to the patcher as the remapping info.
- DuplicatePluginCommandlet now passes in full remapping information to the patcher during incremental uploads where we may only want to patch a subset of all possible files

#rnx

diff --git a/Engine/Source/Developer/AssetTools/Internal/AssetHeaderPatcher.h b/Engine/Source/Developer/AssetTools/Internal/AssetHeaderPatcher.h
--- a/Engine/Source/Developer/AssetTools/Internal/AssetHeaderPatcher.h
+++ b/Engine/Source/Developer/AssetTools/Internal/AssetHeaderPatcher.h
@@ -49,10 +49,17 @@ UE_INTERNAL struct FAssetHeaderPatcher
 		 * be patched due to internal references. The mapping provided in InSrcAndDstPackagePaths will be used 
 		 * to determine the filepath on disk to write when patching.
 		 *
-		 * @param InSrcAndDstPackagePaths Map of all long package names (/Root/Folder/Package) to be patched and to which new name they should be patched to.
-		 * @param bInGatherDependentPackages If true (default), upon creating the context GatherDependentPackages() will be called.
+		 * @param	InSrcAndDstPackagePaths Map of all long package names (/Root/Folder/Package) that are patchable 
+					and to which new name they should be patched to. Note, this should be all possible remappings
+					so that references to unpatched assets may still be corrected when patching assets. If a subset
+					of packages should be written to disk during patching, pass that list in InSrcPackagePathsToPatch.
+		 * @param	bInGatherDependentPackages If true (default), upon creating the context GatherDependentPackages() will be called.
+		 * @param	InSrcPackagePathsToPatch Optional. List of long package names (/Root/Folder/Package) to be patched on disk. 
+					This set of package names must be a subset or complete set of the keys in InSrcAndDstPackagePaths. If not
+					provided, the keys in InSrcAndDstPackagePaths will be used instead.
 		 **/
-		UE_API FContext(const TMap<FString, FString>& InSrcAndDstPackagePaths, const bool bInGatherDependentPackages = true);
+		UE_API FContext(const TMap<FString, FString>& InSrcAndDstPackagePaths, const bool bInGatherDependentPackages = true,
+			const TArray<FString>* InSrcPackagePathsToPatch = nullptr);
 
 		/**
 		 * Context used for patching. Contains all information for how object and package references
@@ -61,15 +68,25 @@ UE_INTERNAL struct FAssetHeaderPatcher
 		 * When patching, package paths to patch will be deduced by the filepath mappings provided in InSrcAndDstFilePaths. All assets
 		 * under InSrcRoot will be written as package paths under a mountpoint located at InSrcBaseDir.
 		 * 
-		 * e.g. Path "C:/User/Repo/Project/Content/Skeletons/Player.uasset" -> "/InSrcRoot/Skeletons/Player" when InSrcBaseDir=C:/User/Repo/Project (/Content is assumed internally)
+		 * e.g. Path "C:/User/Repo/Project/Content/Skeletons/Player.uasset" -> "/InSrcRoot/Skeletons/Player" 
+				when InSrcBaseDir=C:/User/Repo/Project (/Content is assumed internally)
 		 * 
-		 * @param InSrcRoot The root mount point for assets to be patched
-		 * @param InDstRoot The new root mount point for patched assets to be placed under
-		 * @param InSrcBaseDir Path to the directory holding the /Content/ directory for assets to patch
-		 * @param InSrcAndDstFilePaths Map of filepaths for files to be patched and where to write the patched version to
-		 * @param InMountPointReplacements Map of root mountpoints (name only, no "/" prefix or suffix) to replace when patching
+		 * @param	InSrcRoot The root mount point for assets to be patched
+		 * @param	InDstRoot The new root mount point for patched assets to be placed under
+		 * @param	InSrcBaseDir Path to the directory holding the /Content/ directory for assets to patch
+		 * @param	InSrcAndDstFilePaths Map of filepaths of all files possible to be patched and where to write the patched 
+					version to. Note, this should be all possible remappings so that references to unpatched assets may still 
+					be corrected when patching assets. If a subset of files should be written to disk during patching, 
+					pass that list in InSrcFilePathsToPatch.
+		 * @param	InMountPointReplacements Map of root mountpoints (name only, no "/" prefix or suffix) to replace when patching
+		 * @param	InSrcFilePathsToPatch Optional. List of filepaths for the files to be patched on disk. This set of files must be a 
+					subset or complete set of the keys in InSrcAndDstFilePaths. If not provided the keys in InSrcAndDstFilePaths 
+					will be used instead.
 		 **/
-		UE_API FContext(const FString& InSrcRoot, const FString& InDstRoot, const FString& InSrcBaseDir, const TMap<FString, FString>& InSrcAndDstFilePaths, const TMap<FString, FString>& InMountPointReplacements);
+		UE_API FContext(const FString& InSrcRoot, const FString& InDstRoot, const FString& InSrcBaseDir,
+			const TMap<FString, FString>& InSrcAndDstFilePathsForRemapping,
+			const TMap<FString, FString>& InMountPointReplacements,
+			const TArray<FString>* InSrcFilePathsToPatch = nullptr);
 
 		/*
 		* Returns the mapping of source long package names to destination package paths used when patching.
@@ -87,13 +104,18 @@ UE_INTERNAL struct FAssetHeaderPatcher
 
 		UE_API void AddVerseMounts();
 		UE_API void GatherDependentPackages();
-		UE_API void GenerateFilePathsFromPackagePaths();
-		UE_API void GeneratePackagePathsFromFilePaths(const FString& InSrcRoot, const FString& InDstRoot, const FString& InSrcBaseDir);
+		UE_API void GenerateFilePathsFromPackagePaths(const TArray<FString>* InFilePathsToPatch = nullptr);
+		UE_API void GeneratePackagePathsFromFilePaths(const FString& InSrcRoot, const FString& InDstRoot, 
+			const FString& InSrcBaseDir, const TArray<FString>* InFilePathsToPatch = nullptr);
 		UE_API void GenerateAdditionalRemappings();
 	
 		TArray<FString> VerseMountPoints;
+		// PackagePath remapping information to use while patching
 		TMap<FString, FString> PackagePathRenameMap;
+		// FilePath remapping information to use while patching
 		TMap<FString, FString> FilePathRenameMap;
+		// Map of Src -> Dst for files on disk to patch
+		TMap<FString, FString> FilePathsToPatchMap;
 
 		// Todo: Make TSet once FCoreRedirect GetTypeHash is implemented
 		TArray<FCoreRedirect> Redirects;
diff --git a/Engine/Source/Developer/AssetTools/Private/AssetHeaderPatcher.cpp b/Engine/Source/Developer/AssetTools/Private/AssetHeaderPatcher.cpp
--- a/Engine/Source/Developer/AssetTools/Private/AssetHeaderPatcher.cpp
+++ b/Engine/Source/Developer/AssetTools/Private/AssetHeaderPatcher.cpp
@@ -114,7 +114,7 @@ FString LexToString(FAssetHeaderPatcher::EResult InResult)
 	}
 }
 
-FAssetHeaderPatcher::FContext::FContext(const TMap<FString, FString>& SourceAndDestPackages, const bool bInGatherDependentPackages)
+FAssetHeaderPatcher::FContext::FContext(const TMap<FString, FString>& SourceAndDestPackages, const bool bInGatherDependentPackages, const TArray<FString>* InSrcPackagePathsToPatch)
 	: PackagePathRenameMap(SourceAndDestPackages)
 {
 	AddVerseMounts();
@@ -124,16 +124,42 @@ FAssetHeaderPatcher::FContext::FContext(const TMap<FString, FString>& SourceAndD
 		GatherDependentPackages();
 	}
 
-	GenerateFilePathsFromPackagePaths();
+	GenerateFilePathsFromPackagePaths(InSrcPackagePathsToPatch);
 	GenerateAdditionalRemappings();
 }
 
-FAssetHeaderPatcher::FContext::FContext(const FString& InSrcRoot, const FString& InDstRoot, const FString& InSrcBaseDir, const TMap<FString, FString>& InSrcAndDstFilePaths, const TMap<FString, FString>& InMountPointReplacements)
+FAssetHeaderPatcher::FContext::FContext(const FString& InSrcRoot, const FString& InDstRoot, const FString& InSrcBaseDir, 
+	const TMap<FString, FString>& InSrcAndDstFilePaths, const TMap<FString, FString>& InMountPointReplacements,
+	const TArray<FString>* InSrcFilePathsToPatch)
 	: FilePathRenameMap(InSrcAndDstFilePaths)
 	, StringMountReplacements(InMountPointReplacements)
 {
+
+	// If we have a patch file list, verify we have mappings for the passed in files and update our FilePathsToPatchMap.
+	// Otherwise use the complete mapping information from InSrcAndDstFilePaths.
+	if (!InSrcFilePathsToPatch)
+	{
+		FilePathsToPatchMap = InSrcAndDstFilePaths;
+	}
+	else
+	{
+		FilePathsToPatchMap.Reserve(InSrcFilePathsToPatch->Num());
+		for (const FString& SrcFile : *InSrcFilePathsToPatch)
+		{
+			const FString* DstFile = InSrcAndDstFilePaths.Find(SrcFile);
+			if (!DstFile)
+			{
+				UE_LOG(LogAssetHeaderPatcher, Warning, TEXT("No remapping for file %s was provided when patching. This file will not be patched."), *SrcFile);
+			}
+			else
+			{
+				FilePathsToPatchMap.Emplace(SrcFile, *DstFile);
+			}
+		}
+	}
+
 	AddVerseMounts();
-	GeneratePackagePathsFromFilePaths(InSrcRoot, InDstRoot, InSrcBaseDir);
+	GeneratePackagePathsFromFilePaths(InSrcRoot, InDstRoot, InSrcBaseDir, InSrcFilePathsToPatch);
 	GenerateAdditionalRemappings();
 }
 
@@ -143,19 +169,23 @@ void FAssetHeaderPatcher::FContext::AddVerseMounts()
 	VerseMountPoints.Add(TEXT("localhost"));
 }
 
-void FAssetHeaderPatcher::FContext::GenerateFilePathsFromPackagePaths()
+void FAssetHeaderPatcher::FContext::GenerateFilePathsFromPackagePaths(const TArray<FString>* InFilePathsToPatch)
 {
+	if (InFilePathsToPatch)
+	{
+		FilePathsToPatchMap.Reserve(InFilePathsToPatch->Num());
+	}
 	FilePathRenameMap.Reserve(PackagePathRenameMap.Num());
 
 	// Construct all source and destination filenames from our package map
 	for (const TTuple<FString, FString>& Package : PackagePathRenameMap)
 	{
-		const FString& PackageName = Package.Key;
+		const FString& SrcPackageName = Package.Key;
 		const FString& DestPackage = Package.Value;
 		FString SrcFilename;
 
 		// To consider: Allow the caller to provide their own file filter
-		if (FPackageName::IsVersePackage(PackageName))
+		if (FPackageName::IsVersePackage(SrcPackageName))
 		{
 			// Verse packages are not header patchable.
 			// They are also not Packages as far as DoesPackageExist tells me.
@@ -163,19 +193,38 @@ void FAssetHeaderPatcher::FContext::GenerateFilePathsFromPackagePaths()
 			continue;
 		}
 
-		if (FPackageName::DoesPackageExist(PackageName, &SrcFilename))
+		if (FPackageName::DoesPackageExist(SrcPackageName, &SrcFilename))
 		{
 			FString DestFilename = FPackageName::LongPackageNameToFilename(DestPackage, FString(FPathViews::GetExtension(SrcFilename, true)));
-			FilePathRenameMap.Add({ MoveTemp(SrcFilename), MoveTemp(DestFilename) });
+			FilePathRenameMap.Add({ SrcFilename, DestFilename });
+
+			// If we have a patch file list, verify we have mappings for the passed in packages and update our FilePathsToPatchMap.
+			// Otherwise use the complete mapping information from PackagePathRenameMap.
+			if (InFilePathsToPatch)
+			{
+				if (!InFilePathsToPatch->Contains(SrcPackageName))
+				{
+					UE_LOG(LogAssetHeaderPatcher, Warning, TEXT("No remapping for package %s was provided. This package will not be patched."), *SrcPackageName);
+				}
+				else
+				{
+					FilePathsToPatchMap.Add({ SrcFilename, DestFilename });
+				}
+			}
+			else
+			{
+				// Patch all files we have a remapping for whose src asset exists on disk
+				FilePathsToPatchMap.Add({ SrcFilename, DestFilename });
+			}
 		}
 		else
 		{
-			UE_LOG(LogAssetHeaderPatcher, Warning, TEXT("{%s} package does not exist, and will not be patched."), *PackageName);
+			UE_LOG(LogAssetHeaderPatcher, Warning, TEXT("{%s} package does not exist, and will not be patched."), *SrcPackageName);
 		}
 	}
 }
 
-void FAssetHeaderPatcher::FContext::GeneratePackagePathsFromFilePaths(const FString& InSrcRoot, const FString& InDstRoot, const FString& InSrcBaseDir)
+void FAssetHeaderPatcher::FContext::GeneratePackagePathsFromFilePaths(const FString& InSrcRoot, const FString& InDstRoot, const FString& InSrcBaseDir, const TArray<FString>* InFilePathsToPatch)
 {
 	const FString SourceContentPath = FPaths::Combine(InSrcBaseDir, TEXT("Content"));
 	for (const TTuple<FString, FString>& SourceAndDest : FilePathRenameMap)
@@ -1074,7 +1123,7 @@ UE::Tasks::FTask FAssetHeaderPatcher::PatchAsync(int32* InOutNumFilesToPatch, in
 
 UE::Tasks::FTask FAssetHeaderPatcher::PatchAsync(int32* InOutNumFilesToPatch, int32* InOutNumFilesPatched, FAssetHeaderPatcherCompletionDelegate InOnSuccess, FAssetHeaderPatcherCompletionDelegate InOnError)
 {
-	PatchedFiles = Context.FilePathRenameMap;
+	PatchedFiles = Context.FilePathsToPatchMap;
 	if (InOutNumFilesToPatch)
 	{
 		*InOutNumFilesToPatch = PatchedFiles.Num();
