SourceControl: Only perform Login() call if necessary.

#rnx

diff --git a/Engine/Source/Developer/AssetTools/Private/AssetFixUpRedirectors.cpp b/Engine/Source/Developer/AssetTools/Private/AssetFixUpRedirectors.cpp
--- a/Engine/Source/Developer/AssetTools/Private/AssetFixUpRedirectors.cpp
+++ b/Engine/Source/Developer/AssetTools/Private/AssetFixUpRedirectors.cpp
@@ -622,12 +622,14 @@ void FAssetFixUpRedirectors::ExecuteFixUp(TArray<TWeakObjectPtr<UObjectRedirecto
 	{
 		ISourceControlProvider& SourceControlProvider = ISourceControlModule::Get().GetProvider();
 		// Update the source control server availability to make sure we can do the rename operation
-		SourceControlProvider.Login();
 		if (!SourceControlProvider.IsAvailable())
 		{
-			// We have failed to update source control even though it is enabled. This is critical and we can not continue
-			FMessageDialog::Open( EAppMsgType::Ok, NSLOCTEXT("UnrealEd", "SourceControl_ServerUnresponsive", "Revision Control is unresponsive. Please check your connection and try again.") );
-			return;
+			if (SourceControlProvider.Login() != ECommandResult::Succeeded)
+			{
+				// We have failed to connect to source control even though it is enabled. This is critical and we can not continue
+				FMessageDialog::Open( EAppMsgType::Ok, NSLOCTEXT("UnrealEd", "SourceControl_ServerUnresponsive", "Revision Control is unresponsive. Please check your connection and try again.") );
+				return;
+			}
 		}
 
 		TArray<UPackage*> PackagesToAddToSCCUpdate;
@@ -635,7 +637,13 @@ void FAssetFixUpRedirectors::ExecuteFixUp(TArray<TWeakObjectPtr<UObjectRedirecto
 		{
 			PackagesToAddToSCCUpdate.Add(RedirectorRefs.Redirector->GetOutermost());
 		}
-		SourceControlProvider.Execute(ISourceControlOperation::Create<FUpdateStatus>(), PackagesToAddToSCCUpdate);
+
+		if (SourceControlProvider.Execute(ISourceControlOperation::Create<FUpdateStatus>(), PackagesToAddToSCCUpdate) != ECommandResult::Succeeded)
+		{
+			// We have failed to update source control even though it is enabled. This is critical and we can not continue
+			FMessageDialog::Open( EAppMsgType::Ok, NSLOCTEXT("UnrealEd", "SourceControl_UpdateFailed", "Revision Control update failed. Please check your connection and try again.") );
+			return;
+		}
 	}
 
 	// Load all referencing packages.
diff --git a/Engine/Source/Developer/AssetTools/Private/AssetRenameManager.cpp b/Engine/Source/Developer/AssetTools/Private/AssetRenameManager.cpp
--- a/Engine/Source/Developer/AssetTools/Private/AssetRenameManager.cpp
+++ b/Engine/Source/Developer/AssetTools/Private/AssetRenameManager.cpp
@@ -876,8 +876,11 @@ bool FAssetRenameManager::UpdatePackageStatus(TArray<FAssetRenameDataWithReferen
 	{
 		ISourceControlProvider& SourceControlProvider = ISourceControlModule::Get().GetProvider();
 
-		// Update the source control server availability to make sure we can do the rename operation
-		SourceControlProvider.Login();
+		// Update the source control server availability if necessary to make sure we can do the rename operation
+		if (!SourceControlProvider.IsAvailable())
+		{
+			SourceControlProvider.Login();
+		}
 
 		if (SourceControlProvider.IsAvailable())
 		{
