Fixed malformed LOCTEXT literal

#rnx

diff --git a/Engine/Source/Developer/AssetTools/Private/AssetFixUpRedirectors.cpp b/Engine/Source/Developer/AssetTools/Private/AssetFixUpRedirectors.cpp
--- a/Engine/Source/Developer/AssetTools/Private/AssetFixUpRedirectors.cpp
+++ b/Engine/Source/Developer/AssetTools/Private/AssetFixUpRedirectors.cpp
@@ -491,8 +491,7 @@ private:
 	FText GetDeleteExplanationText() const
 	{
 		return LOCTEXT("RedirectorDeleteExplanation", 
-			"Delete all redirectors which are no longer referenced locally. " LINE_TERMINATOR 
-			LINE_TERMINATOR
+			"Delete all redirectors which are no longer referenced locally.\n\n"
 			"Use with caution as these redirectors may be referenced"
 			" - for example in other source control branches or in other users' workspaces." 
 			);
@@ -501,8 +500,7 @@ private:
 	FText GetKeepExplanationText() const
 	{
 		return LOCTEXT("RedirectorKeepExplanation", 
-			"Keep all redirectors, both referenced and unreferenced." LINE_TERMINATOR
-			LINE_TERMINATOR
+			"Keep all redirectors, both referenced and unreferenced.\n\n"
 			"This prevents issues caused by references to them in other workspaces or source control branches.");
 	}
 
