AssetHeaderPatcher fixes
- `PatchingTask` member wasn't set in `PatchAsync` so `IsPatching` would always return false
- Added `ETaskPriority` param and set it to `BackgroundHigh` in `AdvancedCopyPackages` otherwise the main thread ends up doing a lot of the work and has no opportunity to update the progress dialog
- Introduced `FPatchAsyncParams` struct to simplify `PatchAsync` signature
- `FAssetHeaderPatcher::FContext` now has a destination base directory which is required to honor asset renaming that go beyond mount points as dictated by `SrcAndDstFilePathsForRemapping` (not used or required by existing code paths but will be used in an upcoming changelist)

#rnx

diff --git a/Engine/Source/Developer/AssetTools/Internal/AssetHeaderPatcher.h b/Engine/Source/Developer/AssetTools/Internal/AssetHeaderPatcher.h
--- a/Engine/Source/Developer/AssetTools/Internal/AssetHeaderPatcher.h
+++ b/Engine/Source/Developer/AssetTools/Internal/AssetHeaderPatcher.h
@@ -16,6 +16,7 @@
  */
 DECLARE_DELEGATE_TwoParams(FAssetHeaderPatcherCompletionDelegate, const FString& /*SrcFilePath*/, const FString& /*DstFilePath*/)
 
+
 UE_INTERNAL struct FAssetHeaderPatcher
 {
 	static const EUnrealEngineObjectUE5Version MinimumSupportedUE5FileVersion = EUnrealEngineObjectUE5Version::ADD_SOFTOBJECTPATH_LIST;
@@ -58,14 +59,16 @@ UE_INTERNAL struct FAssetHeaderPatcher
 					This set of package names must be a subset or complete set of the keys in InSrcAndDstPackagePaths. If not
 					provided, the keys in InSrcAndDstPackagePaths will be used instead.
 		 **/
-		UE_API FContext(const TMap<FString, FString>& InSrcAndDstPackagePaths, const bool bInGatherDependentPackages = true,
+		UE_API FContext(
+			const TMap<FString, FString>& InSrcAndDstPackagePaths, 
+			const bool bInGatherDependentPackages = true,
 			const TArray<FString>* InSrcPackagePathsToPatch = nullptr);
 
 		/**
 		 * Context used for patching. Contains all information for how object and package references
 		 * will be changed as part of patching. 
 		 * 
-		 * When patching, package paths to patch will be deduced by the filepath mappings provided in InSrcAndDstFilePaths. All assets
+		 * When patching, package paths to patch will be deduced by the filepath mappings provided in InSrcAndDstFilePathsForRemapping. All assets
 		 * under InSrcRoot will be written as package paths under a mountpoint located at InSrcBaseDir.
 		 * 
 		 * e.g. Path "C:/User/Repo/Project/Content/Skeletons/Player.uasset" -> "/InSrcRoot/Skeletons/Player" 
@@ -74,16 +77,20 @@ UE_INTERNAL struct FAssetHeaderPatcher
 		 * @param	InSrcRoot The root mount point for assets to be patched
 		 * @param	InDstRoot The new root mount point for patched assets to be placed under
 		 * @param	InSrcBaseDir Path to the directory holding the /Content/ directory for assets to patch
-		 * @param	InSrcAndDstFilePaths Map of filepaths of all files possible to be patched and where to write the patched 
+		 * @param	InSrcAndDstFilePathsForRemapping Map of filepaths of all files possible to be patched and where to write the patched 
 					version to. Note, this should be all possible remappings so that references to unpatched assets may still 
 					be corrected when patching assets. If a subset of files should be written to disk during patching, 
 					pass that list in InSrcFilePathsToPatch.
 		 * @param	InMountPointReplacements Map of root mountpoints (name only, no "/" prefix or suffix) to replace when patching
 		 * @param	InSrcFilePathsToPatch Optional. List of filepaths for the files to be patched on disk. This set of files must be a 
-					subset or complete set of the keys in InSrcAndDstFilePaths. If not provided the keys in InSrcAndDstFilePaths 
+					subset or complete set of the keys in InSrcAndDstFilePathsForRemapping. If not provided the keys in InSrcAndDstFilePathsForRemapping 
 					will be used instead.
 		 **/
-		UE_API FContext(const FString& InSrcRoot, const FString& InDstRoot, const FString& InSrcBaseDir,
+		UE_API FContext(
+			const FString& InSrcRoot,
+			const FString& InDstRoot,
+			const FString& InSrcBaseDir,
+			const FString& InDstBaseDir,
 			const TMap<FString, FString>& InSrcAndDstFilePathsForRemapping,
 			const TMap<FString, FString>& InMountPointReplacements,
 			const TArray<FString>* InSrcFilePathsToPatch = nullptr);
@@ -105,8 +112,7 @@ UE_INTERNAL struct FAssetHeaderPatcher
 		UE_API void AddVerseMounts();
 		UE_API void GatherDependentPackages();
 		UE_API void GenerateFilePathsFromPackagePaths(const TArray<FString>* InFilePathsToPatch = nullptr);
-		UE_API void GeneratePackagePathsFromFilePaths(const FString& InSrcRoot, const FString& InDstRoot, 
-			const FString& InSrcBaseDir, const TArray<FString>* InFilePathsToPatch = nullptr);
+		UE_API void GeneratePackagePathsFromFilePaths(const FString& InSrcRoot, const FString& InDstRoot, const FString& InSrcBaseDir, const FString& InDstBaseDir);
 		UE_API void GenerateAdditionalRemappings();
 	
 		TArray<FString> VerseMountPoints;
@@ -136,15 +142,38 @@ UE_INTERNAL struct FAssetHeaderPatcher
 	*/
 	UE_API void SetContext(const FContext& InContext);
 
+	/*
+	* PatchAsync optional parameters
+	*/
+	struct FPatchAsyncParams
+	{
+		/** Optional value to know how many files are expected to be read/written during patching */
+		int32* OutNumFilesToPatch;
+
+		/** Optional value used to know how the patcher is progressing (useful for progress bars) */
+		int32* OutNumFilesPatched;
+
+		/** Thread-safe delegate invoked for each successful file copy */
+		FAssetHeaderPatcherCompletionDelegate OnSuccess;
+
+		/** Thread-safe delegate invoked for each failed file copy */
+		FAssetHeaderPatcherCompletionDelegate OnError;
+
+		/** Task priority */
+		UE::Tasks::ETaskPriority TaskPriority;
+
+		FPatchAsyncParams()
+			: OutNumFilesToPatch(nullptr)
+			, OutNumFilesPatched(nullptr)
+			, TaskPriority(UE::Tasks::ETaskPriority::Default)
+		{}
+	};
+
 	/*
 	* Schedules the reading of source files determined by the patcher context, as well as the writing of the patched versions
-	* of all source files read. 
-	* 
-	* @param InOutNumFilesToPatch Optional value to know how many files are expected to be read/written during patching
-	* @param InOutNumFilesPatched Optional value used to know how the patcher is progressing (useful for progress bars)
+	* of all source files read.
 	*/
-	UE_API UE::Tasks::FTask PatchAsync(int32* InOutNumFilesToPatch = nullptr, int32* InOutNumFilesPatched = nullptr);
-	UE_API UE::Tasks::FTask PatchAsync(int32* InOutNumFilesToPatch, int32* InOutNumFilesPatched, FAssetHeaderPatcherCompletionDelegate InOnSuccess, FAssetHeaderPatcherCompletionDelegate InOnError);
+	UE_API UE::Tasks::FTask PatchAsync(const FPatchAsyncParams& Params = {});
 
 	/*
 	* Returns the status of any inflight patching operations. In the case of multiple errors, the last seen error will be reported.
diff --git a/Engine/Source/Developer/AssetTools/Private/AssetHeaderPatcher.cpp b/Engine/Source/Developer/AssetTools/Private/AssetHeaderPatcher.cpp
--- a/Engine/Source/Developer/AssetTools/Private/AssetHeaderPatcher.cpp
+++ b/Engine/Source/Developer/AssetTools/Private/AssetHeaderPatcher.cpp
@@ -128,8 +128,13 @@ FAssetHeaderPatcher::FContext::FContext(const TMap<FString, FString>& SourceAndD
 	GenerateAdditionalRemappings();
 }
 
-FAssetHeaderPatcher::FContext::FContext(const FString& InSrcRoot, const FString& InDstRoot, const FString& InSrcBaseDir, 
-	const TMap<FString, FString>& InSrcAndDstFilePaths, const TMap<FString, FString>& InMountPointReplacements,
+FAssetHeaderPatcher::FContext::FContext(
+	const FString& InSrcRoot,
+	const FString& InDstRoot,
+	const FString& InSrcBaseDir,
+	const FString& InDstBaseDir,
+	const TMap<FString, FString>& InSrcAndDstFilePaths,
+	const TMap<FString, FString>& InMountPointReplacements,
 	const TArray<FString>* InSrcFilePathsToPatch)
 	: FilePathRenameMap(InSrcAndDstFilePaths)
 	, StringMountReplacements(InMountPointReplacements)
@@ -159,7 +164,7 @@ FAssetHeaderPatcher::FContext::FContext(const FString& InSrcRoot, const FString&
 	}
 
 	AddVerseMounts();
-	GeneratePackagePathsFromFilePaths(InSrcRoot, InDstRoot, InSrcBaseDir, InSrcFilePathsToPatch);
+	GeneratePackagePathsFromFilePaths(InSrcRoot, InDstRoot, InSrcBaseDir, InDstBaseDir);
 	GenerateAdditionalRemappings();
 }
 
@@ -224,23 +229,36 @@ void FAssetHeaderPatcher::FContext::GenerateFilePathsFromPackagePaths(const TArr
 	}
 }
 
-void FAssetHeaderPatcher::FContext::GeneratePackagePathsFromFilePaths(const FString& InSrcRoot, const FString& InDstRoot, const FString& InSrcBaseDir, const TArray<FString>* InFilePathsToPatch)
+void FAssetHeaderPatcher::FContext::GeneratePackagePathsFromFilePaths(
+	const FString& InSrcRoot,
+	const FString& InDstRoot,
+	const FString& InSrcBaseDir,
+	const FString& InDstBaseDir)
 {
-	const FString SourceContentPath = FPaths::Combine(InSrcBaseDir, TEXT("Content"));
+	const FString SrcContentPath = FPaths::Combine(InSrcBaseDir, TEXTVIEW("Content"));
+	const FString DstContentPath = FPaths::Combine(InDstBaseDir, TEXTVIEW("Content"));
 	for (const TTuple<FString, FString>& SourceAndDest : FilePathRenameMap)
 	{
 		const FString& SrcFileName = SourceAndDest.Key;
+		const FString& DstFileName = SourceAndDest.Value;
 
-		if (FPaths::IsUnderDirectory(SrcFileName, SourceContentPath))
+		FStringView SrcRelativePkgPath;
+		FStringView DstRelativePkgPath;
+		if (FPathViews::TryMakeChildPathRelativeTo(SrcFileName, SrcContentPath, SrcRelativePkgPath) &&
+			FPathViews::TryMakeChildPathRelativeTo(DstFileName, DstContentPath, DstRelativePkgPath))
 		{
-			if (FStringView RelativePkgPath; FPathViews::TryMakeChildPathRelativeTo(SrcFileName, SourceContentPath, RelativePkgPath))
+			// chop the extension
+			SrcRelativePkgPath = FPathViews::GetBaseFilenameWithPath(SrcRelativePkgPath);
+			DstRelativePkgPath = FPathViews::GetBaseFilenameWithPath(DstRelativePkgPath);
+
+			if (!SrcRelativePkgPath.IsEmpty() &&
+				!DstRelativePkgPath.IsEmpty() &&
+				!SrcRelativePkgPath.EndsWith(TEXTVIEW("/")) &&
+				!DstRelativePkgPath.EndsWith(TEXTVIEW("/")))
 			{
-				RelativePkgPath = FPathViews::GetBaseFilenameWithPath(RelativePkgPath); // chop the extension
-				if (RelativePkgPath.Len() > 0 && !RelativePkgPath.EndsWith(TEXT("/")))
-				{
-					PackagePathRenameMap.Add(FPaths::Combine(TEXT("/"), InSrcRoot, RelativePkgPath),
-						FPaths::Combine(TEXT("/"), InDstRoot, RelativePkgPath));
-				}
+				PackagePathRenameMap.Add(
+					FPaths::Combine(TEXTVIEW("/"), InSrcRoot, SrcRelativePkgPath),
+					FPaths::Combine(TEXTVIEW("/"), InDstRoot, DstRelativePkgPath));
 			}
 		}
 	}
@@ -1116,21 +1134,15 @@ void FAssetHeaderPatcher::SetContext(const FContext& InContext)
 	Reset();
 }
 
-UE::Tasks::FTask FAssetHeaderPatcher::PatchAsync(int32* InOutNumFilesToPatch, int32* InOutNumFilesPatched)
-{
-	return PatchAsync(InOutNumFilesToPatch, InOutNumFilesPatched, FAssetHeaderPatcherCompletionDelegate(), FAssetHeaderPatcherCompletionDelegate());
-}
-
-UE::Tasks::FTask FAssetHeaderPatcher::PatchAsync(int32* InOutNumFilesToPatch, int32* InOutNumFilesPatched, FAssetHeaderPatcherCompletionDelegate InOnSuccess, FAssetHeaderPatcherCompletionDelegate InOnError)
+UE::Tasks::FTask FAssetHeaderPatcher::PatchAsync(const FPatchAsyncParams& Params)
 {
 	PatchedFiles = Context.FilePathsToPatchMap;
-	if (InOutNumFilesToPatch)
+	if (Params.OutNumFilesToPatch)
 	{
-		*InOutNumFilesToPatch = PatchedFiles.Num();
+		*Params.OutNumFilesToPatch = PatchedFiles.Num();
 	}
 
 	// Spawn tasks (Scatter)
-	UE::Tasks::FTask PatchAssetsCleanupTask;
 	TArray<UE::Tasks::FTask> PatchAssetTasks;
 
 	// Note we are scheduling and launching tasks one at a time rather than preparing all jobs and launching all at once.
@@ -1146,9 +1158,9 @@ UE::Tasks::FTask FAssetHeaderPatcher::PatchAsync(int32* InOutNumFilesToPatch, in
 		auto DoPatchFn = [this, 
 			SrcFilename = Filename.Key,
 			DstFilename = Filename.Value,
-			NumPatched = InOutNumFilesPatched, 
-			OnSuccess = InOnSuccess,
-			OnError = InOnError]()
+			NumPatched = Params.OutNumFilesPatched,
+			OnSuccess = Params.OnSuccess,
+			OnError = Params.OnError]()
 			{
 				// Even if we are cancelled, increment our progress
 				if (NumPatched)
@@ -1187,12 +1199,12 @@ UE::Tasks::FTask FAssetHeaderPatcher::PatchAsync(int32* InOutNumFilesToPatch, in
 		}
 		else
 		{
-			PatchAssetTasks.Add(UE::Tasks::Launch(UE_SOURCE_LOCATION, MoveTemp(DoPatchFn)));
+			PatchAssetTasks.Add(UE::Tasks::Launch(UE_SOURCE_LOCATION, MoveTemp(DoPatchFn), Params.TaskPriority));
 		}
 	}
 
 	// Once all tasks have completed, remove the redirects before we declare Patching complete
-	UE::Tasks::FTask PatcherCleanupTask = UE::Tasks::Launch(UE_SOURCE_LOCATION, [this]()
+	PatchingTask = UE::Tasks::Launch(UE_SOURCE_LOCATION, [this]()
 		{
 			if (Status != EResult::Cancelled && ErroredFiles.IsEmpty())
 			{
@@ -1207,11 +1219,13 @@ UE::Tasks::FTask FAssetHeaderPatcher::PatchAsync(int32* InOutNumFilesToPatch, in
 				}
 			}
 
-		}, UE::Tasks::Prerequisites(PatchAssetTasks));
+		},
+		UE::Tasks::Prerequisites(PatchAssetTasks),
+		bSingleThreaded ? UE::Tasks::ETaskPriority::Default : Params.TaskPriority);
 
 	Status = EResult::InProgress;
 
-	return PatcherCleanupTask;
+	return PatchingTask;
 }
 
 FAssetHeaderPatcher::EResult FAssetHeaderPatcherInner::PatchHeader()
diff --git a/Engine/Source/Developer/AssetTools/Private/AssetTools.cpp b/Engine/Source/Developer/AssetTools/Private/AssetTools.cpp
--- a/Engine/Source/Developer/AssetTools/Private/AssetTools.cpp
+++ b/Engine/Source/Developer/AssetTools/Private/AssetTools.cpp
@@ -2506,7 +2506,12 @@ bool UAssetToolsImpl::AdvancedCopyPackages(
 
 			int32 NumFilesToPatch = 0;
 			int32 PatchAssetsCompletedCount = 0;
-			UE::Tasks::FTask PatchAssetsCompletionTask = Patcher.PatchAsync(&NumFilesToPatch, &PatchAssetsCompletedCount);
+			FAssetHeaderPatcher::FPatchAsyncParams PatchParams;
+			PatchParams.OutNumFilesToPatch = &NumFilesToPatch;
+			PatchParams.OutNumFilesPatched = &PatchAssetsCompletedCount;
+			PatchParams.TaskPriority = UE::Tasks::ETaskPriority::BackgroundHigh;
+
+			UE::Tasks::FTask PatchAssetsCompletionTask = Patcher.PatchAsync(PatchParams);
 		
 			LoopProgress = MakeUnique<FScopedSlowTask>(static_cast<float>(NumFilesToPatch), LOCTEXT("AdvancedCopyPackages.ReplacingAssetReferences", "Replacing Asset References..."));
 			LoopProgress->MakeDialog();
