Added support for showing Verse paths to rename failure reasons.

#rnx

diff --git a/Engine/Source/Developer/AssetTools/Private/AssetRenameManager.cpp b/Engine/Source/Developer/AssetTools/Private/AssetRenameManager.cpp
--- a/Engine/Source/Developer/AssetTools/Private/AssetRenameManager.cpp
+++ b/Engine/Source/Developer/AssetTools/Private/AssetRenameManager.cpp
@@ -2,6 +2,7 @@
 
 #include "AssetRenameManager.h"
 #include "AssetDefinitionRegistry.h"
+#include "Containers/VersePath.h"
 #include "Serialization/ArchiveUObject.h"
 #include "UObject/Class.h"
 #include "Misc/PackageName.h"
@@ -1374,10 +1375,12 @@ void FAssetRenameManager::SetupPublicAssets(TArray<FAssetRenameDataWithReference
 		}
 
 		// Nothing to do if the renamed package mount point doesn't change
-		const FNameBuilder OldPackageName(RenameData.OldObjectPath.GetLongPackageFName());
-		const FNameBuilder NewPackageName(RenameData.NewObjectPath.GetLongPackageFName());
-		const FStringView OldPackageMountPoint = FPathViews::GetMountPointNameFromPath(OldPackageName);
-		const FStringView NewPackageMountPoint = FPathViews::GetMountPointNameFromPath(NewPackageName);
+		const FName OldPackageName = RenameData.OldObjectPath.GetLongPackageFName();
+		const FName NewPackageName = RenameData.NewObjectPath.GetLongPackageFName();
+		const FNameBuilder OldPackageNameBuilder(OldPackageName);
+		const FNameBuilder NewPackageNameBuilder(NewPackageName);
+		const FStringView OldPackageMountPoint = FPathViews::GetMountPointNameFromPath(OldPackageNameBuilder);
+		const FStringView NewPackageMountPoint = FPathViews::GetMountPointNameFromPath(NewPackageNameBuilder);
 		if (NewPackageMountPoint.Equals(OldPackageMountPoint, ESearchCase::IgnoreCase))
 		{
 			continue;
@@ -1416,8 +1419,8 @@ void FAssetRenameManager::SetupPublicAssets(TArray<FAssetRenameDataWithReference
 							RenameData.bRenameFailed = true;
 							RenameData.FailureReason = FText::Format(
 								LOCTEXT("MovedAssetReferencingPrivateAsset", "Cannot move asset to {0} because it would be referencing a private asset in a different plugin: {1}"),
-								FText::FromStringView(NewPackageMountPoint),
-								FText::FromName(DependencyPackageName));
+								FText::FromString(AssetTools.GetUserFacingPluginName(NewPackageMountPoint)),
+								FText::FromString(AssetTools.GetUserFacingLongPackageName(DependencyPackageName)));
 							break;
 						}
 					}
@@ -1444,12 +1447,12 @@ void FAssetRenameManager::SetupPublicAssets(TArray<FAssetRenameDataWithReference
 		}
 
 		// Figure out if it's gonna be referenced from another mount point after the move
-		FString ReferencingAssetInDifferentMountPoint;
+		FName ReferencingAssetInDifferentMountPoint;
 		if (RenameData.bCreateRedirector)
 		{
 			ReferencingAssetInDifferentMountPoint = OldPackageName;
 		}
-		if (ReferencingAssetInDifferentMountPoint.IsEmpty())
+		if (ReferencingAssetInDifferentMountPoint.IsNone())
 		{
 			for (FName It : RenameData.NotRenamedReferencingPackageNames)
 			{
@@ -1457,12 +1460,12 @@ void FAssetRenameManager::SetupPublicAssets(TArray<FAssetRenameDataWithReference
 				const FStringView ReferencingPackagMountPoint = FPathViews::GetMountPointNameFromPath(ReferencingPackageName);
 				if (!ReferencingPackagMountPoint.Equals(NewPackageMountPoint, ESearchCase::IgnoreCase))
 				{
-					ReferencingAssetInDifferentMountPoint = ReferencingPackageName;
+					ReferencingAssetInDifferentMountPoint = It;
 					break;
 				}
 			}
 		}
-		if (ReferencingAssetInDifferentMountPoint.IsEmpty())
+		if (ReferencingAssetInDifferentMountPoint.IsNone())
 		{
 			for (const TPair<FName, FName>& It : RenameData.RenamedReferencingPackageNames)
 			{
@@ -1476,27 +1479,28 @@ void FAssetRenameManager::SetupPublicAssets(TArray<FAssetRenameDataWithReference
 				}
 				FAssetRenameDataWithReferencers& ReferencingRenameData = **FoundReferencingRenameData;
 
-				FNameBuilder ReferencingPackageName;
+				FName ReferencingPackageName;
 				if (ReferencingRenameData.bRenameFailed || ReferencingRenameData.bOnlyFixSoftReferences)
 				{
-					OldReferencingPackageName.AppendString(ReferencingPackageName);
+					ReferencingPackageName = OldReferencingPackageName;
 				}
 				else
 				{
-					NewReferencingPackageName.AppendString(ReferencingPackageName);
+					ReferencingPackageName = NewReferencingPackageName;
 				}
 
-				const FStringView ReferencingPackagMountPoint = FPathViews::GetMountPointNameFromPath(ReferencingPackageName);
+				const FNameBuilder ReferencingPackageNameBuilder(ReferencingPackageName);
+				const FStringView ReferencingPackagMountPoint = FPathViews::GetMountPointNameFromPath(ReferencingPackageNameBuilder);
 				if (!NewPackageMountPoint.Equals(ReferencingPackagMountPoint, ESearchCase::IgnoreCase))
 				{
-					ReferencingAssetInDifferentMountPoint = ReferencingPackageName.ToString();
+					ReferencingAssetInDifferentMountPoint = ReferencingPackageName;
 					break;
 				}
 			}
 		}
 
 		// Check if the asset can be made public
-		if (!ReferencingAssetInDifferentMountPoint.IsEmpty())
+		if (!ReferencingAssetInDifferentMountPoint.IsNone())
 		{
 			const FString OldAssetPath = RenameData.OldObjectPath.GetAssetPathString();
 			if (AssetTools.CanAssetBePublic(OldAssetPath))
@@ -1511,14 +1515,14 @@ void FAssetRenameManager::SetupPublicAssets(TArray<FAssetRenameDataWithReference
 				{
 					RenameData.FailureReason = FText::Format(
 						LOCTEXT("AssetCannotBePublicForRedirector", "Cannot move asset to {0} because it cannot be made public in order to be referenced from a different plugin by its redirector"),
-						FText::FromStringView(NewPackageMountPoint));
+						FText::FromStringView(AssetTools.GetUserFacingPluginName(NewPackageMountPoint)));
 				}
 				else
 				{
 					RenameData.FailureReason = FText::Format(
 						LOCTEXT("AssetCannotBePublic", "Cannot move asset to {0} because it cannot be made public in order to be referenced from a different plugin by {1}"),
-						FText::FromStringView(NewPackageMountPoint),
-						FText::FromString(ReferencingAssetInDifferentMountPoint));
+						FText::FromStringView(AssetTools.GetUserFacingPluginName(NewPackageMountPoint)),
+						FText::FromString(AssetTools.GetUserFacingLongPackageName(ReferencingAssetInDifferentMountPoint)));
 				}
 			}
 		}
@@ -2057,9 +2061,25 @@ int32 FAssetRenameManager::ReportFailures(const TArray<FAssetRenameDataWithRefer
 			UObject* Asset = RenameData.Asset.Get();
 			if (Asset)
 			{
+				UE::Core::FVersePath VersePath;
+				if (IAssetTools::Get().ShowingContentVersePath())
+				{
+					VersePath = Asset->GetVersePath();
+				}
+
+				FText AssetName;
+				if (VersePath.IsValid())
+				{
+					AssetName = MoveTemp(VersePath).AsText();
+				}
+				else
+				{
+					AssetName = FText::FromName(Asset->GetOutermost()->GetFName());
+				}
+
 				FFormatNamedArguments Args;
 				Args.Add(TEXT("FailureReason"), RenameData.FailureReason);
-				Args.Add(TEXT("AssetName"), FText::FromString(Asset->GetOutermost()->GetName()));
+				Args.Add(TEXT("AssetName"), MoveTemp(AssetName));
 
 				FailedRenames.Add(FText::Format(LOCTEXT("AssetRenameFailure", "{AssetName} - {FailureReason}"), Args));
 			}
