When moving an asset to another plugin, don't change the access specifier if the asset is Epic internal, only make it public if it was private.

#rnx

diff --git a/Engine/Source/Developer/AssetTools/Private/AssetRenameManager.cpp b/Engine/Source/Developer/AssetTools/Private/AssetRenameManager.cpp
--- a/Engine/Source/Developer/AssetTools/Private/AssetRenameManager.cpp
+++ b/Engine/Source/Developer/AssetTools/Private/AssetRenameManager.cpp
@@ -495,7 +495,7 @@ bool FAssetRenameManager::FixReferencesAndRename(const TArray<FAssetRenameData>&
 			DetectReadOnlyPackages(AssetsToRename, ReferencingPackagesToSave);
 
 
-			// Make public any asset that will be referenced from another plugin after the rename.
+			// Make public any private asset that will be referenced from another plugin after the rename.
 			// If the asset cannot be made public or if moving it requires a referenced asset that's not being modified to become public,
 			// its rename will fail and other assets that have dependencies between them will also fail to be renamed.
 			SetupPublicAssets(AssetsToRename);
@@ -1434,8 +1434,8 @@ void FAssetRenameManager::SetupPublicAssets(TArray<FAssetRenameDataWithReference
 			}
 		}
 
-		// If the asset is already public, it can be referenced from anywhere so we're good
-		if (Package->GetAssetAccessSpecifier() == EAssetAccessSpecifier::Public)
+		// If the asset isn't private, it can be referenced externally so we're good
+		if (Package->GetAssetAccessSpecifier() != EAssetAccessSpecifier::Private)
 		{
 			continue;
 		}
diff --git a/Engine/Source/Developer/AssetTools/Private/AssetRenameManager.h b/Engine/Source/Developer/AssetTools/Private/AssetRenameManager.h
--- a/Engine/Source/Developer/AssetTools/Private/AssetRenameManager.h
+++ b/Engine/Source/Developer/AssetTools/Private/AssetRenameManager.h
@@ -111,7 +111,7 @@ private:
 	void DetectReadOnlyPackages(TArray<FAssetRenameDataWithReferencers>& AssetsToRename, TArray<UPackage*>& InOutReferencingPackagesToSave) const;
 
 	/** 
-	 * Make public any asset that will be referenced from another plugin after the rename.
+	 * Make public any private asset that will be referenced from another plugin after the rename.
 	 * If the asset cannot be made public or if moving it requires a referenced asset that's not being modified to become public,
 	 * its rename will fail and other assets that have dependencies between them will also fail to be renamed.
 	 */
