Fixed selecting a specific file format instead of "All Files" triggers a Legacy Import
- Solution: In ImportAssetsWithDialogImplementation, if Interchange is enabled do not select factory based on extension, let ImportAssetsInternal decide otherwise Interchange is skipped.

#rnx

diff --git a/Engine/Source/Developer/AssetTools/Private/AssetTools.cpp b/Engine/Source/Developer/AssetTools/Private/AssetTools.cpp
--- a/Engine/Source/Developer/AssetTools/Private/AssetTools.cpp
+++ b/Engine/Source/Developer/AssetTools/Private/AssetTools.cpp
@@ -6440,7 +6440,8 @@ TArray<UObject*> UAssetToolsImpl::ImportAssetsWithDialogImplementation(const FSt
 		if (OpenFilenames.Num() > 0)
 		{
 			UFactory* ChosenFactory = nullptr;
-			if (FilterIndex > 0)
+			// A filter may have been chosen but most likely for convenience, let ImportAssets decipher that.
+			if (!UInterchangeManager::IsInterchangeImportEnabled() && FilterIndex > 0)
 			{
 				ChosenFactory = *FilterIndexToFactory.Find(FilterIndex);
 			}
