Interchange - Fixed trivial typo in comment detected by BugHawk after submitting CL-46119713.

#rnx

diff --git a/Engine/Source/Developer/AssetTools/Private/AssetTools.cpp b/Engine/Source/Developer/AssetTools/Private/AssetTools.cpp
--- a/Engine/Source/Developer/AssetTools/Private/AssetTools.cpp
+++ b/Engine/Source/Developer/AssetTools/Private/AssetTools.cpp
@@ -3869,7 +3869,7 @@ TArray<UObject*> UAssetToolsImpl::ImportAssetsInternal(const TArray<FString>& Fi
 								}
 							}
 
-							const IConsoleVariable* Var = IConsoleManager::Get().FindConsoleVariable(TEXT("Interchange.FeatureFlags.Import.SyncToBrowser")); // Declared in InterchangeEngineModule.cpp
+							const IConsoleVariable* Var = IConsoleManager::Get().FindConsoleVariable(TEXT("Interchange.FeatureFlags.Import.SyncToBrowser")); // Declared in InterchangeEditorModule.cpp
 
 							// If the CVar is found, it overrides the passed value.
 							const bool bResolvedSyncToBrowser = Var ? Var->GetBool() : bSyncToBrowser;
