Updated maximum asset name message to more accurately reflect the character limit.

#rnx

diff --git a/Engine/Source/Developer/AssetTools/Private/AssetViewUtils.cpp b/Engine/Source/Developer/AssetTools/Private/AssetViewUtils.cpp
--- a/Engine/Source/Developer/AssetTools/Private/AssetViewUtils.cpp
+++ b/Engine/Source/Developer/AssetTools/Private/AssetViewUtils.cpp
@@ -1455,7 +1455,7 @@ bool AssetViewUtils::IsValidObjectPathForCreate(const FString& ObjectPath, const
 	if ( ObjectPath.Len() >= NAME_SIZE )
 	{
 		OutErrorMessage = FText::Format(LOCTEXT("AssetNameTooLong", "The name for the asset is too long, the maximum is '{0}'. \nPlease choose a shorter name for the asset or create it in a shallower folder structure."),
-			FText::AsNumber(NAME_SIZE));
+			FText::AsNumber(NAME_SIZE - 1));
 		// Return false to indicate that the user should enter a new name
 		return false;
 	}
@@ -1544,7 +1544,7 @@ bool AssetViewUtils::IsValidFolderPathForCreate(const FString& InFolderPath, con
 	if (NewFolderPath.Len() >= NAME_SIZE)
 	{
 		OutErrorMessage = FText::Format(LOCTEXT("PackagePathTooLong", "The path for the folder is too deep, the maximum is '{0}'. \nPlease choose a shorter name for the folder or create it in a shallower folder structure."),
-			FText::AsNumber(NAME_SIZE));
+			FText::AsNumber(NAME_SIZE - 1));
 		// Return false to indicate that the user should enter a new name
 		return false;
 	}
