Added support for showing Verse paths when prompting to fixed up packages because they have soft references to a renamed object.

#rnx

diff --git a/Engine/Source/Developer/AssetTools/Private/AssetRenameManager.cpp b/Engine/Source/Developer/AssetTools/Private/AssetRenameManager.cpp
--- a/Engine/Source/Developer/AssetTools/Private/AssetRenameManager.cpp
+++ b/Engine/Source/Developer/AssetTools/Private/AssetRenameManager.cpp
@@ -518,9 +518,9 @@ bool FAssetRenameManager::FixReferencesAndRename(const TArray<FAssetRenameData>&
 					}
 
 					FString AssetNames;
-					for (UPackage* PackageToSave : ReferencingPackagesToSave)
+					for (const FString& AssetName : IAssetTools::Get().GetUserFacingLongPackageNames(ReferencingPackagesToSave))
 					{
-						AssetNames += FString("\n") + PackageToSave->GetName();
+						AssetNames += FString("\n") + AssetName;
 					}
 
 					// Warn user before saving referencing packages
