Added support for Asset Header Patching for ImportTypeHierarchy entries.

#rnx

diff --git a/Engine/Source/Developer/AssetTools/Private/AssetHeaderPatcher.cpp b/Engine/Source/Developer/AssetTools/Private/AssetHeaderPatcher.cpp
--- a/Engine/Source/Developer/AssetTools/Private/AssetHeaderPatcher.cpp
+++ b/Engine/Source/Developer/AssetTools/Private/AssetHeaderPatcher.cpp
@@ -17,6 +17,7 @@
 #include "Misc/PathViews.h"
 #include "Misc/RedirectCollector.h"
 #include "Serialization/LargeMemoryReader.h"
+#include "Serialization/ImportTypeHierarchy.h"
 #include "UObject/CoreRedirects.h"
 #include "UObject/CoreRedirects/CoreRedirectsContext.h"
 #include "UObject/Linker.h"
@@ -740,6 +741,7 @@ enum class EPatchedSection
 	ImportTable,
 	ExportTable,
 	SoftPackageReferencesTable,
+	ImportTypeHierarchyMap,
 	ThumbnailTable,
 	AssetRegistryData,
 	AssetRegistryDependencyData,
@@ -765,6 +767,7 @@ enum class ESummaryOffset
 	DependsTable,
 	SoftPackageReferenceList,
 	SearchableNamesMap,
+	ImportTypeHierarchyMap,
 	ThumbnailTable,
 	AssetRegistryData,
 	WorldTileInfoData,
@@ -876,6 +879,7 @@ void PatchSummaryOffsets(FPackageFileSummary& Dst, int64 OffsetFrom, int64 Offse
 				{ UE_POPULATE_OFFSET_INFO(DependsOffset) },
 				{ UE_POPULATE_OFFSET_INFO(SoftPackageReferencesOffset) },
 				{ UE_POPULATE_OFFSET_INFO(SearchableNamesOffset) },
+				{ UE_POPULATE_OFFSET_INFO(ImportTypeHierarchiesOffset) },
 				{ UE_POPULATE_OFFSET_INFO(ThumbnailTableOffset) },
 				{ UE_POPULATE_OFFSET_INFO(AssetRegistryDataOffset) },
 				{ UE_POPULATE_OFFSET_INFO(BulkDataStartOffset) },
@@ -957,6 +961,7 @@ public:
 		SoftObjectPathTable.Empty();
 		SoftPackageReferencesTable.Empty();
 		Summary = FPackageFileSummary();
+		ImportTypeHierarchyMap.Empty();
 		ThumbnailTable.Empty();
 		UnchangedNames.Empty();
 	}
@@ -967,6 +972,7 @@ public:
 	bool DoPatch(FTopLevelAssetPath& InOutPath);
 	bool DoPatch(FGatherableTextData& InOutGatherablerTextData);
 	bool DoPatch(FThumbnailEntry& InOutThumbnailEntry);
+	bool DoPatch(UE::Serialization::Private::FTypeResource& InOutTypeResource);
 	bool RemapFName(FName SrcName, FName DstName);
 	bool AddFName(FName DstName);
 	bool ShouldReplaceMountPoint(const FStringView InPath, FStringView& OutSrcMountPoint, FStringView& OutDstMountPoint);
@@ -1024,6 +1030,7 @@ public:
 		int64 SoftPackageReferencesListSize = -1;
 		int64 ThumbnailTableSize = -1;
 		int64 SearchableNamesMapSize = -1;
+		int64 ImportTypeHierarchiesMapSize = -1;
 		int64 AssetRegistryDataSize = -1;
 		int64 PackageTrailerSize = -1;
 	};
@@ -1055,6 +1062,7 @@ public:
 	TArray<FName> SoftPackageReferencesTable;
 	TMap<FPackageIndex, TArray<FName>> SearchableNamesMap;
 	TArray<FThumbnailEntry> ThumbnailTable;
+	TMap<FPackageIndex, UE::Serialization::Private::FImportTypeHierarchy> ImportTypeHierarchyMap;
 
 	// Asset registry data information
 	struct FAssetRegistryObjectData
@@ -1540,6 +1548,15 @@ FAssetHeaderPatcher::EResult FAssetHeaderPatcherInner::PatchHeader_Deserialize()
 		HeaderInformation.SearchableNamesMapSize = MemAr.Tell() - Summary.SearchableNamesOffset;
 	}
 
+	if (Summary.ImportTypeHierarchiesOffset)
+	{
+		MemAr.Seek(Summary.ImportTypeHierarchiesOffset);
+
+		MemAr << ImportTypeHierarchyMap;
+
+		HeaderInformation.ImportTypeHierarchiesMapSize = MemAr.Tell() - Summary.ImportTypeHierarchiesOffset;
+	}
+
 	if (Summary.ThumbnailTableOffset)
 	{
 		MemAr.Seek(Summary.ThumbnailTableOffset);
@@ -2559,6 +2576,58 @@ bool FAssetHeaderPatcherInner::DoPatch(FThumbnailEntry& InThumbnailEntry)
 	return bPatched;
 }
 
+bool FAssetHeaderPatcherInner::DoPatch(UE::Serialization::Private::FTypeResource& InOutTypeResource)
+{
+	auto PatchObjectName = [](FName& Name, const FCoreRedirectObjectName& Src, const FCoreRedirectObjectName& Dst, FAssetHeaderPatcherInner& Patcher) -> bool
+	{
+		if (Patcher.RemapFName(Src.ObjectName, Dst.ObjectName))
+		{
+			Name = Dst.ObjectName;
+			return true;
+		}
+		else
+		{
+			Name = Src.ObjectName;
+			return false;
+		}
+	};
+
+	auto PatchPackageName = [](FName& Name, const FCoreRedirectObjectName& Src, const FCoreRedirectObjectName& Dst, FAssetHeaderPatcherInner& Patcher) -> bool
+	{
+		if (Patcher.RemapFName(Src.PackageName, Dst.PackageName))
+		{
+			Name = Dst.PackageName;
+			return true;
+		}
+		else
+		{
+			Name = Src.PackageName;
+			return false;
+		}
+	};
+
+	bool bDidPatch = false;
+
+	// Note that even though we are patching the super types stored in the type hierarchy, the *actual* type hierarchy
+	// might be different (because e.g., types have moved). Note however, that the type hierarchy itself is a snapshot
+	// of what the import types looked liked at the point the asset was saved. For this reason, we don't attempt to
+	// create the type hierarchy again here. This will be done the next type the patched asset is saved.
+
+	const FCoreRedirectObjectName SrcClass(InOutTypeResource.ClassName, NAME_None, InOutTypeResource.ClassPackageName);
+	const FCoreRedirectObjectName DstClass = FCoreRedirects::GetRedirectedName(ECoreRedirectFlags::Type_Class | ECoreRedirectFlags::Type_Package, SrcClass, ECoreRedirectMatchFlags::DisallowPartialLHSMatch);
+
+	bDidPatch |= PatchObjectName(InOutTypeResource.ClassName, SrcClass, DstClass, *this);
+	bDidPatch |= PatchPackageName(InOutTypeResource.ClassPackageName, SrcClass, DstClass, *this);
+
+	const FCoreRedirectObjectName SrcObj(InOutTypeResource.TypeName, NAME_None, InOutTypeResource.PackageName);
+	const FCoreRedirectObjectName DstObj = FCoreRedirects::GetRedirectedName(ECoreRedirectFlags::Type_Object | ECoreRedirectFlags::Type_Package, SrcObj, ECoreRedirectMatchFlags::DisallowPartialLHSMatch);
+
+	bDidPatch |= PatchObjectName(InOutTypeResource.TypeName, SrcObj, DstObj, *this);
+	bDidPatch |= PatchPackageName(InOutTypeResource.PackageName, SrcObj, DstObj, *this);
+
+	return bDidPatch;
+}
+
 FAssetHeaderPatcher::EResult FAssetHeaderPatcherInner::PatchHeader_PatchSections()
 {
 	// Package Summary
@@ -2631,6 +2700,21 @@ FAssetHeaderPatcher::EResult FAssetHeaderPatcherInner::PatchHeader_PatchSections
 		}
 	}
 
+	// ImportTypeHierarchy Map
+	for (auto& Pair : ImportTypeHierarchyMap)
+	{
+		UE::Serialization::Private::FImportTypeHierarchy& Entry = Pair.Value;
+
+		const TArray<UE::Serialization::Private::FTypeResource>& SuperTypes = Entry.GetSuperTypes();
+
+		for (const UE::Serialization::Private::FTypeResource& SuperType : SuperTypes)
+		{
+			UE::Serialization::Private::FTypeResource& WritableSuperType = const_cast<UE::Serialization::Private::FTypeResource&>(SuperType);
+
+			DoPatch(WritableSuperType);
+		}
+	}
+
 	// Thumbnail Table
 	for (FThumbnailEntry& ThumbnailEntry : ThumbnailTable)
 	{
@@ -2836,6 +2920,7 @@ FAssetHeaderPatcher::EResult FAssetHeaderPatcherInner::PatchHeader_WriteDestinat
 		{ EPatchedSection::ExportTable,					Summary.ExportOffset,					HeaderInformation.ExportTableSize,					true	},
 		{ EPatchedSection::SoftPackageReferencesTable,	Summary.SoftPackageReferencesOffset,	HeaderInformation.SoftPackageReferencesListSize,	false	},
 		{ EPatchedSection::SearchableNamesMap,			Summary.SearchableNamesOffset,			HeaderInformation.SearchableNamesMapSize,			false	},
+		{ EPatchedSection::ImportTypeHierarchyMap,		Summary.ImportTypeHierarchiesOffset,	HeaderInformation.ImportTypeHierarchiesMapSize,		false	},
 		{ EPatchedSection::ThumbnailTable,				Summary.ThumbnailTableOffset,			HeaderInformation.ThumbnailTableSize,				false	},
 		{ EPatchedSection::AssetRegistryData,			Summary.AssetRegistryDataOffset,		AssetRegistryData.SectionSize,						true	},
 		{ EPatchedSection::AssetRegistryDependencyData,	AssetRegistryData.PkgData.DependencyDataOffset,	AssetRegistryData.DependencyDataSectionSize,		false	},
@@ -3070,6 +3155,22 @@ FAssetHeaderPatcher::EResult FAssetHeaderPatcherInner::PatchHeader_WriteDestinat
 			break;
 		}
 
+		case EPatchedSection::ImportTypeHierarchyMap:
+		{
+			const int64 MapStartOffset = Writer.Tell();
+			Writer << ImportTypeHierarchyMap;
+
+			checkf(!Writer.IsCriticalError(), TEXT("Issue writing %s"), *Writer.GetErrorMessage());
+
+			const int64 MapSize = Writer.Tell() - MapStartOffset;
+			const int64 Delta = MapSize - SourceSection.Size;
+			checkf(Delta == 0, TEXT("Delta should be Zero. is %d"), (int)Delta);
+			check(Summary.ImportTypeHierarchiesCount == ImportTypeHierarchyMap.Num());
+			check(Summary.ImportTypeHierarchiesOffset == MapStartOffset);
+
+			break;
+		}
+
 		case EPatchedSection::ThumbnailTable:
 		{
 			const int64 ThumbnailTableStartOffset = Writer.Tell();
@@ -3471,6 +3572,44 @@ void FAssetHeaderPatcherInner::DumpState(FStringView OutputDirectory)
 	Builder.RemoveSuffix(1);
 	Builder.Append(TEXT("\n\t],\n"));
 
+	Builder.Append(TEXT("\t\"ImportTypeHierarchyMap\":[ "));
+	for (TPair<FPackageIndex, UE::Serialization::Private::FImportTypeHierarchy> Pair : ImportTypeHierarchyMap)
+	{
+		Builder.Append(TEXT("\n\t\t{\n"));
+		Builder.Append(TEXT("\t\t\t\"Super types\": ["));
+
+		const TArray<UE::Serialization::Private::FTypeResource>& SuperTypes = Pair.Value.GetSuperTypes();
+
+		for (const UE::Serialization::Private::FTypeResource& SuperType : SuperTypes)
+		{
+			Builder.Append(TEXT("\n\t\t\t\t{"));
+
+			Builder.Append(TEXT("\n\t\t\t\t\t\"TypeName\": \""));
+			Builder.Append(GetDebugFNameString(SuperType.TypeName));
+			Builder.Append(TEXT("\","));
+
+			Builder.Append(TEXT("\n\t\t\t\t\t\"PackageName\": \""));
+			Builder.Append(GetDebugFNameString(SuperType.PackageName));
+			Builder.Append(TEXT("\","));
+
+			Builder.Append(TEXT("\n\t\t\t\t\t\"ClassName\": \""));
+			Builder.Append(GetDebugFNameString(SuperType.ClassName));
+			Builder.Append(TEXT("\","));
+
+			Builder.Append(TEXT("\n\t\t\t\t\t\"ClassPackageName\": \""));
+			Builder.Append(GetDebugFNameString(SuperType.ClassPackageName));
+			Builder.Append(TEXT("\","));
+
+			Builder.Append(TEXT("\n\t\t\t\t},"));
+		}
+
+		Builder.RemoveSuffix(1);
+		Builder.Append(TEXT("\n\t\t\t]"));
+		Builder.Append(TEXT("\n\t\t},"));
+	}
+	Builder.RemoveSuffix(1);
+	Builder.Append(TEXT("\n\t],\n"));
+
 	Builder.Append(TEXT("\t\"ThumbnailTable\":[ "));
 	for (const FThumbnailEntry& ThumbnailEntry : ThumbnailTable)
 	{
@@ -4728,6 +4867,37 @@ TEST_CASE_NAMED(FAssetHeaderPatcherTests, "AssetHeaderPatcher", "[AssetHeaderPat
 		}
 	}
 
+	SECTION("DoPatch(FTypeResource)")
+	{
+		ResetPatcher();
+		TArray<FCoreRedirect> ImportTypeHierarchyRedirects =
+		{
+			{ ECoreRedirectFlags::Type_Package, TEXT("/Engine/SrcClassPackage"),                           TEXT("/Engine/DstClassPackage")},
+			{ ECoreRedirectFlags::Type_Class,   TEXT("/Engine/SrcClassPackage.SrcClass"),                  TEXT("/Engine/DstClassPackage.DstClass")},
+			{ ECoreRedirectFlags::Type_Package, TEXT("/Module/_Verse/VerseAsset"),                         TEXT("/Module/_Verse") },
+			{ ECoreRedirectFlags::Type_Object,  TEXT("/Module/_Verse/VerseAsset.some_verse_class"),        TEXT("/Module/_Verse.VerseAsset-some_verse_class") },
+		};
+		CHECK(FCoreRedirects::AddRedirectList(ImportTypeHierarchyRedirects, TEXT("Asset Header Patcher Tests - FTypeResource Patching")));
+
+		UE::Serialization::Private::FTypeResource Actual;
+		Actual.TypeName = SrcVerseImportObjectName;
+		Actual.PackageName = SrcVerseAssetName;
+		Actual.ClassName = SrcImportClassName;
+		Actual.ClassPackageName = SrcImportClassPackage;
+
+		UE::Serialization::Private::FTypeResource Expected;
+		Expected.TypeName = DstVerseImportObjectName;
+		Expected.PackageName = DstVerseAssetName;
+		Expected.ClassName = DstImportClassName;
+		Expected.ClassPackageName = DstImportClassPackage;
+
+		CHECK(Patcher.DoPatch(Actual));
+		CHECK_EQUALS(TEXT("FTypeResource patching doesn't update TypeName"), Actual.TypeName, Expected.TypeName);
+		CHECK_EQUALS(TEXT("FTypeResource patching doesn't update PackageName"), Actual.PackageName, Expected.PackageName);
+		CHECK_EQUALS(TEXT("FTypeResource patching doesn't update ClassName"), Actual.ClassName, Expected.ClassName);
+		CHECK_EQUALS(TEXT("FTypeResource patching doesn't update ClassPackageName"), Actual.ClassPackageName, Expected.ClassPackageName);
+	}
+
 	// Import Test Breakdown
 	// 
 	// Assume that SrcEngineModule and DstEngineModule both exist and types are moving between them and/or possibly being renamed while moving
