AssetTools: Add API to Asset Tools to get the default name for an asset of a given class and with given a source object. Modify CreateAssetsFrom() to use the new function.

#rnx

diff --git a/Engine/Source/Developer/AssetTools/Private/AssetTools.cpp b/Engine/Source/Developer/AssetTools/Private/AssetTools.cpp
--- a/Engine/Source/Developer/AssetTools/Private/AssetTools.cpp
+++ b/Engine/Source/Developer/AssetTools/Private/AssetTools.cpp
@@ -1841,7 +1841,7 @@ void UAssetToolsImpl::CreateAssetsFrom(TConstArrayView<UObject*> SourceObjects,
 			if (UFactory* Factory = FactoryConstructor(SourceObject))
 			{
 				// Create an appropriate and unique name 
-				const TOptional<FString> DefaultAssetName = GetDefaultAssetNameForClass(CreateAssetType);
+				TOptional<FString> DefaultAssetName = GetDefaultAssetNameForClassFromSource(CreateAssetType, SourceObject);
 
 				const FString BasePackageName = DefaultAssetName.IsSet() ? DefaultAssetName.GetValue() : SourceObject->GetName();
 				const FString Suffix = DefaultAssetName.IsSet() ? FString() : DefaultSuffix;
@@ -1868,7 +1868,7 @@ void UAssetToolsImpl::CreateAssetsFrom(TConstArrayView<UObject*> SourceObjects,
 				if (UFactory* Factory = FactoryConstructor(SourceObject))
 				{
 					// Determine an appropriate name
-					const TOptional<FString> DefaultAssetName = GetDefaultAssetNameForClass(CreateAssetType);
+					const TOptional<FString> DefaultAssetName = GetDefaultAssetNameForClassFromSource(CreateAssetType, SourceObject);
 
 					const FString BasePackageName = DefaultAssetName.IsSet() ? DefaultAssetName.GetValue() : SourceObject->GetName();
 					const FString Suffix = DefaultAssetName.IsSet() ? FString() : DefaultSuffix;
@@ -7096,6 +7096,22 @@ TOptional<FString> UAssetToolsImpl::GetDefaultAssetNameForClass(const UClass* Cl
 	return TOptional<FString>();
 }
 
+TOptional<FString> UAssetToolsImpl::GetDefaultAssetNameForClassFromSource(const UClass* Class, const UObject* SourceObject)
+{
+	if (FString* AssetName = DefaultNewAssetNames.Find(Class))
+	{
+		FString DefaultAssetName = *AssetName;
+		DefaultAssetName = DefaultAssetName.Replace(TEXT("{source}"), *SourceObject->GetName(), ESearchCase::IgnoreCase);
+
+		SanitizeName(DefaultAssetName);
+
+		FString SantizedDefaultName = ObjectTools::SanitizeInvalidChars(DefaultAssetName, INVALID_LONGPACKAGE_CHARACTERS INVALID_OBJECTNAME_CHARACTERS);
+		return SantizedDefaultName;
+	}
+
+	return TOptional<FString>();
+}
+
 void UAssetToolsImpl::RegisterDefaultAssetNameForClass(const UClass* Class, const FString& DefaultName)
 {
 	DefaultNewAssetNames.Add(Class, DefaultName);
diff --git a/Engine/Source/Developer/AssetTools/Private/AssetTools.h b/Engine/Source/Developer/AssetTools/Private/AssetTools.h
--- a/Engine/Source/Developer/AssetTools/Private/AssetTools.h
+++ b/Engine/Source/Developer/AssetTools/Private/AssetTools.h
@@ -178,6 +178,7 @@ public:
 	virtual void RegisterSanitizeNameDelegate(const FName OwnerName, FSanitizeName Delegate) override;
 	virtual void UnregisterSanitizeNameDelegate(const FName OwnerName) override;
 	virtual TOptional<FString> GetDefaultAssetNameForClass(const UClass* Class) override;
+	virtual TOptional<FString> GetDefaultAssetNameForClassFromSource(const UClass* Class, const UObject* SourceObject) override;
 	virtual void RegisterDefaultAssetNameForClass(const UClass* Class, const FString& DefaultName) override;
 	virtual void UnregisterDefaultAssetNameForClass(const UClass* Class) override;
 	virtual void RegisterCanMigrateAsset(const FName OwnerName, UE::AssetTools::FCanMigrateAsset Delegate) override;
diff --git a/Engine/Source/Developer/AssetTools/Public/IAssetTools.h b/Engine/Source/Developer/AssetTools/Public/IAssetTools.h
--- a/Engine/Source/Developer/AssetTools/Public/IAssetTools.h
+++ b/Engine/Source/Developer/AssetTools/Public/IAssetTools.h
@@ -844,6 +844,8 @@ public:
 
 	/** Get the default name for the supplied class (if one is registered) to use when creating new assets of that type */
 	virtual TOptional<FString> GetDefaultAssetNameForClass(const UClass* Class) = 0;
+	/** Get the default name for the supplied class (if one is registered) to use when creating a new asset of that type from the input source object */
+	virtual TOptional<FString> GetDefaultAssetNameForClassFromSource(const UClass* Class, const UObject* SourceObject) = 0;
 	/** Register a default name for the supplied class to use when creating new assets of that type */
 	virtual void RegisterDefaultAssetNameForClass(const UClass* Class, const FString& DefaultName) = 0;
 	/** Remove any previously registered name for the supplied class */
