AssetViewUtils: Deprecate OnSyncStart / OnSyncFinish so we can remove it in the future.

It's confusing and incomplete. It doesn't work for all syncs. Only for a subset of syncs that are triggered via AssetViewUtils.

#rnx

diff --git a/Engine/Source/Developer/AssetTools/Private/AssetViewUtils.cpp b/Engine/Source/Developer/AssetTools/Private/AssetViewUtils.cpp
--- a/Engine/Source/Developer/AssetTools/Private/AssetViewUtils.cpp
+++ b/Engine/Source/Developer/AssetTools/Private/AssetViewUtils.cpp
@@ -1896,7 +1896,9 @@ bool AssetViewUtils::SyncPackagesFromSourceControl(const TArray<FString>& Packag
 	if (PackageNames.Num() > 0)
 	{
 		// Broadcast sync starting...
+		PRAGMA_DISABLE_DEPRECATION_WARNINGS;
 		AssetViewUtils::OnSyncStart().Broadcast();
+		PRAGMA_ENABLE_DEPRECATION_WARNINGS;
 
 		TArray<FString> PackageNamesToSync = PackageNames;
 
@@ -1954,7 +1956,9 @@ bool AssetViewUtils::SyncPackagesFromSourceControl(const TArray<FString>& Packag
 		SCCProvider.Execute(ISourceControlOperation::Create<FUpdateStatus>(), PackageFilenames);
 
 		// Broadcast sync finished...
+		PRAGMA_DISABLE_DEPRECATION_WARNINGS;
 		AssetViewUtils::OnSyncFinish().Broadcast(SyncResult == ECommandResult::Succeeded, &PackageFilenames);
+		PRAGMA_ENABLE_DEPRECATION_WARNINGS;
 
 		// Return result...
 		return (SyncResult == ECommandResult::Succeeded);
@@ -1989,7 +1993,9 @@ static bool SyncPathsFromSourceControl(const FString& Revision, const TArray<FSt
 	if (PathsOnDisk.Num() > 0)
 	{
 		// Broadcast sync starting...
+		PRAGMA_DISABLE_DEPRECATION_WARNINGS;
 		AssetViewUtils::OnSyncStart().Broadcast();
+		PRAGMA_ENABLE_DEPRECATION_WARNINGS;
 
 		ISourceControlProvider& SCCProvider = ISourceControlModule::Get().GetProvider();
 
@@ -2134,7 +2140,9 @@ static bool SyncPathsFromSourceControl(const FString& Revision, const TArray<FSt
 		SCCProvider.Execute(ISourceControlOperation::Create<FUpdateStatus>(), AffectedFiles.Num() > 0 ? AffectedFiles : PathsOnDisk);
 
 		// Broadcast sync finished...
+		PRAGMA_DISABLE_DEPRECATION_WARNINGS;
 		AssetViewUtils::OnSyncFinish().Broadcast(SyncResult == ECommandResult::Succeeded, AffectedFiles.Num() > 0 ? &AffectedFiles : nullptr);
+		PRAGMA_ENABLE_DEPRECATION_WARNINGS;
 
 		// Return result...
 		return (SyncResult == ECommandResult::Succeeded);
diff --git a/Engine/Source/Developer/AssetTools/Public/AssetViewUtils.h b/Engine/Source/Developer/AssetTools/Public/AssetViewUtils.h
--- a/Engine/Source/Developer/AssetTools/Public/AssetViewUtils.h
+++ b/Engine/Source/Developer/AssetTools/Public/AssetViewUtils.h
@@ -282,6 +282,7 @@ namespace AssetViewUtils
 
 	/** Called when a sync from AssetViewUtils starts */
 	DECLARE_MULTICAST_DELEGATE(FOnSyncStart);
+	UE_DEPRECATED(5.8, "This is no longer supported")
 	ASSETTOOLS_API FOnSyncStart& OnSyncStart();
 
 	/** 
@@ -290,5 +291,6 @@ namespace AssetViewUtils
 	 * If Files is nullptr, assume any file might have been changed by the sync.
 	 */
 	DECLARE_MULTICAST_DELEGATE_TwoParams(FOnSyncFinish, bool /*bSucceeded*/, const TArray<FString>* /*Files*/);
+	UE_DEPRECATED(5.8, "This is no longer supported")
 	ASSETTOOLS_API FOnSyncFinish& OnSyncFinish();
 }
