Disabling header parching of ImportTypeHIerarchy entries when their serialization is disabled.

#rnx

diff --git a/Engine/Source/Developer/AssetTools/Private/AssetHeaderPatcher.cpp b/Engine/Source/Developer/AssetTools/Private/AssetHeaderPatcher.cpp
--- a/Engine/Source/Developer/AssetTools/Private/AssetHeaderPatcher.cpp
+++ b/Engine/Source/Developer/AssetTools/Private/AssetHeaderPatcher.cpp
@@ -1550,11 +1550,18 @@ FAssetHeaderPatcher::EResult FAssetHeaderPatcherInner::PatchHeader_Deserialize()
 
 	if (Summary.ImportTypeHierarchiesOffset)
 	{
-		MemAr.Seek(Summary.ImportTypeHierarchiesOffset);
+		if (UE::Serialization::Private::FImportTypeHierarchy::IsSerializationEnabled())
+		{
+			MemAr.Seek(Summary.ImportTypeHierarchiesOffset);
 
-		MemAr << ImportTypeHierarchyMap;
+			MemAr << ImportTypeHierarchyMap;
 
-		HeaderInformation.ImportTypeHierarchiesMapSize = MemAr.Tell() - Summary.ImportTypeHierarchiesOffset;
+			HeaderInformation.ImportTypeHierarchiesMapSize = MemAr.Tell() - Summary.ImportTypeHierarchiesOffset;
+		}
+		else
+		{
+			HeaderInformation.ImportTypeHierarchiesMapSize = 0;
+		}
 	}
 
 	if (Summary.ThumbnailTableOffset)
@@ -3157,17 +3164,19 @@ FAssetHeaderPatcher::EResult FAssetHeaderPatcherInner::PatchHeader_WriteDestinat
 
 		case EPatchedSection::ImportTypeHierarchyMap:
 		{
-			const int64 MapStartOffset = Writer.Tell();
-			Writer << ImportTypeHierarchyMap;
+			if (UE::Serialization::Private::FImportTypeHierarchy::IsSerializationEnabled())
+			{
+				const int64 MapStartOffset = Writer.Tell();
+				Writer << ImportTypeHierarchyMap;
 
-			checkf(!Writer.IsCriticalError(), TEXT("Issue writing %s"), *Writer.GetErrorMessage());
-
-			const int64 MapSize = Writer.Tell() - MapStartOffset;
-			const int64 Delta = MapSize - SourceSection.Size;
-			checkf(Delta == 0, TEXT("Delta should be Zero. is %d"), (int)Delta);
-			check(Summary.ImportTypeHierarchiesCount == ImportTypeHierarchyMap.Num());
-			check(Summary.ImportTypeHierarchiesOffset == MapStartOffset);
+				checkf(!Writer.IsCriticalError(), TEXT("Issue writing %s"), *Writer.GetErrorMessage());
 
+				const int64 MapSize = Writer.Tell() - MapStartOffset;
+				const int64 Delta = MapSize - SourceSection.Size;
+				checkf(Delta == 0, TEXT("Delta should be Zero. is %d"), (int)Delta);
+				check(Summary.ImportTypeHierarchiesCount == ImportTypeHierarchyMap.Num());
+				check(Summary.ImportTypeHierarchiesOffset == MapStartOffset);
+			}
 			break;
 		}
 
diff --git a/Engine/Source/Runtime/CoreUObject/Public/Serialization/ImportTypeHierarchy.h b/Engine/Source/Runtime/CoreUObject/Public/Serialization/ImportTypeHierarchy.h
--- a/Engine/Source/Runtime/CoreUObject/Public/Serialization/ImportTypeHierarchy.h
+++ b/Engine/Source/Runtime/CoreUObject/Public/Serialization/ImportTypeHierarchy.h
@@ -45,7 +45,7 @@ public:
 	UE_INTERNAL void AddSuperType(const FTypeResource& SuperType) { SuperTypes.Add(SuperType); }
 	UE_INTERNAL void AddSuperType(FTypeResource&& SuperType) { SuperTypes.Add(SuperType); }
 
-	UE_INTERNAL static bool IsSerializationEnabled();
+	UE_INTERNAL COREUOBJECT_API static bool IsSerializationEnabled();
 
 	/** I/O functions */
 	friend COREUOBJECT_API FArchive& operator<<(FArchive& Ar, FImportTypeHierarchy& I);
