AssetTools: Add a sanity check for filename validity before returning a default asset name for a specific asset class.

#rnx

diff --git a/Engine/Source/Developer/AssetTools/Private/AssetTools.cpp b/Engine/Source/Developer/AssetTools/Private/AssetTools.cpp
--- a/Engine/Source/Developer/AssetTools/Private/AssetTools.cpp
+++ b/Engine/Source/Developer/AssetTools/Private/AssetTools.cpp
@@ -7089,8 +7089,16 @@ TOptional<FString> UAssetToolsImpl::GetDefaultAssetNameForClass(const UClass* Cl
 		FString DefaultAssetName = *AssetName;
 		SanitizeName(DefaultAssetName);
 
-		FString SantizedDefaultName = ObjectTools::SanitizeInvalidChars(DefaultAssetName, INVALID_LONGPACKAGE_CHARACTERS INVALID_OBJECTNAME_CHARACTERS);
-		return SantizedDefaultName;
+		FString SanitizedDefaultName = ObjectTools::SanitizeInvalidChars(DefaultAssetName, INVALID_LONGPACKAGE_CHARACTERS INVALID_OBJECTNAME_CHARACTERS);
+
+		FText ErrorText;
+		if (!FFileHelper::IsFilenameValidForSaving(SanitizedDefaultName, ErrorText))
+		{
+			UE_LOG(LogAssetTools, Warning, TEXT("The default name '%s' for class '%s' is invalid: %s"), *SanitizedDefaultName, *Class->GetName(), *ErrorText.ToString());
+			return TOptional<FString>();
+		}
+
+		return SanitizedDefaultName;
 	}
 
 	return TOptional<FString>();
@@ -7105,8 +7113,16 @@ TOptional<FString> UAssetToolsImpl::GetDefaultAssetNameForClassFromSource(const
 
 		SanitizeName(DefaultAssetName);
 
-		FString SantizedDefaultName = ObjectTools::SanitizeInvalidChars(DefaultAssetName, INVALID_LONGPACKAGE_CHARACTERS INVALID_OBJECTNAME_CHARACTERS);
-		return SantizedDefaultName;
+		FString SanitizedDefaultName = ObjectTools::SanitizeInvalidChars(DefaultAssetName, INVALID_LONGPACKAGE_CHARACTERS INVALID_OBJECTNAME_CHARACTERS);
+
+		FText ErrorText;
+		if (!FFileHelper::IsFilenameValidForSaving(SanitizedDefaultName, ErrorText))
+		{
+			UE_LOG(LogAssetTools, Warning, TEXT("The default name '%s' for class '%s' is invalid: %s"), *SanitizedDefaultName, *Class->GetName(), *ErrorText.ToString());
+			return TOptional<FString>();
+		}
+
+		return SanitizedDefaultName;
 	}
 
 	return TOptional<FString>();
