[AssetRenameManager] Avoid reporting references from sub-objects of redirected path

#Synced-CL 48233187

diff --git a/Engine/Source/Developer/AssetTools/Private/AssetRenameManager.cpp b/Engine/Source/Developer/AssetTools/Private/AssetRenameManager.cpp
--- a/Engine/Source/Developer/AssetTools/Private/AssetRenameManager.cpp
+++ b/Engine/Source/Developer/AssetTools/Private/AssetRenameManager.cpp
@@ -1683,6 +1683,7 @@ bool FAssetRenameManager::CheckPackageForSoftObjectReferences(UPackage* Package,
 	for (const TPair<FSoftObjectPath, FSoftObjectPath>& Pair : AssetRedirectorMap)
 	{
 		const FString& CheckSubPath = Pair.Key.GetSubPathString();
+		const FTopLevelAssetPath CheckAssetPath = Pair.Key.GetAssetPath();
 
 		// Find where we're going to start iterating
 		int32 Index = Algo::LowerBound(CachedReferences->Keys, Pair.Key, FSoftObjectPathFastLess());
@@ -1692,7 +1693,7 @@ bool FAssetRenameManager::CheckPackageForSoftObjectReferences(UPackage* Package,
 			const FString& SubPath = CachedKey.GetSubPathString();
 
 			// Stop as soon as we're not anymore in the range we're searching
-			if (Pair.Key.GetWithoutSubPath() != CachedKey.GetWithoutSubPath())
+			if (CheckAssetPath != CachedKey.GetAssetPath())
 			{
 				break;
 			}
@@ -1705,12 +1706,22 @@ bool FAssetRenameManager::CheckPackageForSoftObjectReferences(UPackage* Package,
 			//
 			if (IsSubPath(SubPath, CheckSubPath))
 			{
-				bFoundReference = true;
 				for (const FWeakObjectPtr& WeakPtr : *CachedReferences->Map.Find(CachedKey))
 				{
 					UObject* ObjectPtr = WeakPtr.Get();
 					if (ObjectPtr)
 					{
+						// Skip references from sub-objects of the redirected path:
+						// CacheKey is a reference that was found in the package that is a sub-path of the current asset we are validating against (Pair.Key)
+						// but ObjectPtr is the actual object that has that reference and if that ObjectPtr is actually a sub-object of the asset we are validating against (Pair.Key) 
+						// then that sub-object is itself getting redirected so do not report.
+						const FSoftObjectPath ObjectPath(ObjectPtr);
+						if (ObjectPath.GetAssetPath() == CheckAssetPath && IsSubPath(ObjectPath.GetSubPathString(), CheckSubPath))
+						{
+							continue;
+						}
+
+						bFoundReference = true;
 						OutReferencingObjects.FindOrAdd(CachedKey).AddUnique(ObjectPtr);
 					}
 				}
