Properly deprecate AssetTypeActions for DataAssets, the AssetDefinition was already being used.
Fix the implementation of UAssetDefinition_DataAsset to correctly show the native class in the content browser.
Reparent input mapping and input context to use base AssetTypeActions which shows the correct type names.
#Synced-CL 48480368

diff --git a/Engine/Plugins/Editor/EngineAssetDefinitions/Source/Private/AssetDefinition_DataAsset.cpp b/Engine/Plugins/Editor/EngineAssetDefinitions/Source/Private/AssetDefinition_DataAsset.cpp
--- a/Engine/Plugins/Editor/EngineAssetDefinitions/Source/Private/AssetDefinition_DataAsset.cpp
+++ b/Engine/Plugins/Editor/EngineAssetDefinitions/Source/Private/AssetDefinition_DataAsset.cpp
@@ -157,19 +157,14 @@ namespace MenuExtension_DataAsset
 
 FText UAssetDefinition_DataAsset::GetAssetDisplayName(const FAssetData& AssetData) const
 {
-	static const FName NAME_RowStructure(TEXT("RowStructure"));
-	
 	if (AssetData.IsValid())
 	{
-		const FAssetDataTagMapSharedView::FFindTagResult RowStructureTag = AssetData.TagsAndValues.FindTag(NAME_RowStructure);
-		if (RowStructureTag.IsSet())
+		const FAssetDataTagMapSharedView::FFindTagResult NativeClassTag = AssetData.TagsAndValues.FindTag(UDataAsset::NativeClassTag);
+		if (NativeClassTag.IsSet())
 		{
-			// Handle full path names and deprecated short class names
-			const FTopLevelAssetPath ClassPath = FAssetData::TryConvertShortClassNameToPathName(*RowStructureTag.GetValue(), ELogVerbosity::Log);
-
-			if (const UScriptStruct* FoundStruct = UClass::TryFindTypeSlow<UScriptStruct>(ClassPath.ToString(), EFindFirstObjectOptions::ExactClass))
+			if (UClass* FoundClass = FindObjectSafe<UClass>(nullptr, *NativeClassTag.GetValue()))
 			{
-				return FText::Format(LOCTEXT("DataTableWithRowType", "Data Table ({0})"), FoundStruct->GetDisplayNameText());
+				return FText::Format(LOCTEXT("DataAssetWithType", "Data Asset ({0})"), FoundClass->GetDisplayNameText());
 			}
 		}
 	}
diff --git a/Engine/Plugins/EnhancedInput/Source/InputEditor/Private/InputEditorModule.cpp b/Engine/Plugins/EnhancedInput/Source/InputEditor/Private/InputEditorModule.cpp
--- a/Engine/Plugins/EnhancedInput/Source/InputEditor/Private/InputEditorModule.cpp
+++ b/Engine/Plugins/EnhancedInput/Source/InputEditor/Private/InputEditorModule.cpp
@@ -216,7 +216,7 @@ UObject* UInputAction_Factory::FactoryCreateNew(UClass* Class, UObject* InParent
 
 // Asset type actions
 
-class FAssetTypeActions_InputContext : public FAssetTypeActions_DataAsset
+class FAssetTypeActions_InputContext : public FAssetTypeActions_Base
 {
 public:
 	virtual FText GetName() const override { return NSLOCTEXT("AssetTypeActions", "AssetTypeActions_InputMappingContext", "Input Mapping Context"); }
@@ -226,7 +226,7 @@ public:
 	virtual UClass* GetSupportedClass() const override { return UInputMappingContext::StaticClass(); }
 };
 
-class FAssetTypeActions_InputAction : public FAssetTypeActions_DataAsset
+class FAssetTypeActions_InputAction : public FAssetTypeActions_Base
 {
 public:
 	virtual FText GetName() const override { return NSLOCTEXT("AssetTypeActions", "AssetTypeActions_InputAction", "Input Action"); }
diff --git a/Engine/Source/Developer/AssetTools/Private/AssetTypeActions/AssetTypeActions_DataAsset.cpp b/Engine/Source/Developer/AssetTools/Private/AssetTypeActions/AssetTypeActions_DataAsset.cpp
--- a/Engine/Source/Developer/AssetTools/Private/AssetTypeActions/AssetTypeActions_DataAsset.cpp
+++ b/Engine/Source/Developer/AssetTools/Private/AssetTypeActions/AssetTypeActions_DataAsset.cpp
@@ -11,39 +11,11 @@
 
 #define LOCTEXT_NAMESPACE "AssetTypeActions"
 
-static const FName NAME_NativeClass(TEXT("NativeClass"));
-
-class FNewNodeClassFilter : public IClassViewerFilter
-{
-public:
-	FNewNodeClassFilter(UClass* InBaseClass)
-		: BaseClass(InBaseClass)
-	{
-	}
-
-	virtual bool IsClassAllowed(const FClassViewerInitializationOptions& InInitOptions, const UClass* InClass, TSharedRef< FClassViewerFilterFuncs > InFilterFuncs) override
-	{
-		if (InClass != nullptr)
-		{
-			return InClass->IsChildOf(BaseClass);
-		}
-		return false;
-	}
-
-	virtual bool IsUnloadedClassAllowed(const FClassViewerInitializationOptions& InInitOptions, const TSharedRef< const IUnloadedBlueprintData > InUnloadedClassData, TSharedRef< FClassViewerFilterFuncs > InFilterFuncs) override
-	{
-		return InUnloadedClassData->IsChildOf(BaseClass);
-	}
-
-private:
-	UClass* BaseClass;
-};
-
 FText FAssetTypeActions_DataAsset::GetDisplayNameFromAssetData(const FAssetData& AssetData) const
 {
 	if (AssetData.IsValid())
 	{
-		const FAssetDataTagMapSharedView::FFindTagResult NativeClassTag = AssetData.TagsAndValues.FindTag(NAME_NativeClass);
+		const FAssetDataTagMapSharedView::FFindTagResult NativeClassTag = AssetData.TagsAndValues.FindTag(UDataAsset::NativeClassTag);
 		if (NativeClassTag.IsSet())
 		{
 			if (UClass* FoundClass = FindObjectSafe<UClass>(nullptr, *NativeClassTag.GetValue()))
@@ -56,21 +28,6 @@ FText FAssetTypeActions_DataAsset::GetDisplayNameFromAssetData(const FAssetData&
 	return FText::GetEmpty();
 }
 
-void FAssetTypeActions_DataAsset::GetActions(const TArray<UObject*>& InObjects, FToolMenuSection& Section)
-{
-	TArray<TWeakObjectPtr<UDataAsset>> DataAssets = GetTypedWeakObjectPtrs<UDataAsset>(InObjects);
-
-	Section.AddMenuEntry(
-		"DataAsset_ChangeClass",
-		LOCTEXT("DataAsset_ChangeClass", "Convert to Different DataAsset Type"),
-		LOCTEXT("DataAsset_ChangeClassTip", "Change the class these Data Assets are subclassed from."),
-		FSlateIcon(),
-		FUIAction(
-			FExecuteAction::CreateSP(this, &FAssetTypeActions_DataAsset::ExecuteChangeDataAssetClass, DataAssets)
-		)
-	);
-}
-
 void FAssetTypeActions_DataAsset::PerformAssetDiff(UObject* OldAsset, UObject* NewAsset, const FRevisionInfo& OldRevision, const FRevisionInfo& NewRevision) const
 {
 	const TSharedRef<SDetailsDiff> DetailsDiff = SDetailsDiff::CreateDiffWindow(OldAsset, NewAsset, OldRevision, NewRevision, GetSupportedClass());
@@ -81,45 +38,4 @@ void FAssetTypeActions_DataAsset::PerformAssetDiff(UObject* OldAsset, UObject* N
 	}
 }
 
-void FAssetTypeActions_DataAsset::ExecuteChangeDataAssetClass(TArray<TWeakObjectPtr<UDataAsset>> InDataAssets)
-{
-	const FText TitleText = LOCTEXT("DataAsset_PickNewDataAssetClass", "Pick New DataAsset Class");
-	FClassViewerInitializationOptions Options;
-	Options.ClassFilters.Add(MakeShared<FNewNodeClassFilter>(UDataAsset::StaticClass()));
-	UClass* OutNewDataAssetClass = nullptr;
-	const bool bPressedOk = SClassPickerDialog::PickClass(TitleText, Options, OutNewDataAssetClass, UDataAsset::StaticClass());
-
-	if (bPressedOk && OutNewDataAssetClass != nullptr)
-	{
-		for (TWeakObjectPtr<UDataAsset> DataAssetPtr : InDataAssets)
-		{
-			if (UDataAsset* OldDataAsset = DataAssetPtr.Get())
-			{
-				if (OldDataAsset && OldDataAsset->IsValidLowLevel())
-				{
-					FName ObjectName = OldDataAsset->GetFName();
-					UObject* Outer = OldDataAsset->GetOuter();
-					OldDataAsset->Rename(nullptr, GetTransientPackage(), REN_DoNotDirty | REN_DontCreateRedirectors);
-
-					UObject* NewDataAsset = NewObject<UObject>(Outer, OutNewDataAssetClass, ObjectName, OldDataAsset->GetFlags());
-
-					// Migrate Data
-					{
-						UEngine::FCopyPropertiesForUnrelatedObjectsParams CopyOptions;
-						CopyOptions.bNotifyObjectReplacement = true;
-						UEngine::CopyPropertiesForUnrelatedObjects(OldDataAsset, NewDataAsset, CopyOptions);
-					}
-
-					NewDataAsset->MarkPackageDirty();
-
-					// Consolidate or "Replace" the old object with the new object for any living references.
-					bool bShowDeleteConfirmation = false;
-					TArray<UObject*> OldDataAssetArray({ (UObject*)OldDataAsset });
-					ObjectTools::ConsolidateObjects(NewDataAsset, OldDataAssetArray, bShowDeleteConfirmation);
-				}
-			}
-		}
-	}
-}
-
 #undef LOCTEXT_NAMESPACE
diff --git a/Engine/Source/Developer/AssetTools/Public/AssetTypeActions/AssetTypeActions_DataAsset.h b/Engine/Source/Developer/AssetTools/Public/AssetTypeActions/AssetTypeActions_DataAsset.h
--- a/Engine/Source/Developer/AssetTools/Public/AssetTypeActions/AssetTypeActions_DataAsset.h
+++ b/Engine/Source/Developer/AssetTools/Public/AssetTypeActions/AssetTypeActions_DataAsset.h
@@ -10,7 +10,7 @@
 #define UE_API ASSETTOOLS_API
 
 class
-// UE_DEPRECATED(5.2, "The AssetDefinition system is replacing AssetTypeActions and UAssetDefinition_DataAsset replaced this.  Please see the Conversion Guide in AssetDefinition.h")
+UE_DEPRECATED(5.8, "The AssetDefinition system is replacing AssetTypeActions and UAssetDefinition_DataAsset replaced this.  Please see the Conversion Guide in AssetDefinition.h")
 FAssetTypeActions_DataAsset : public FAssetTypeActions_Base
 {
 public:
@@ -20,11 +20,7 @@ public:
 	virtual UClass* GetSupportedClass() const override { return UDataAsset::StaticClass(); }
 	virtual uint32 GetCategories() override { return EAssetTypeCategories::Misc; }
 	UE_API virtual FText GetDisplayNameFromAssetData(const FAssetData& AssetData) const override;
-	UE_API virtual void GetActions(const TArray<UObject*>& InObjects, struct FToolMenuSection& Section) override;
 	UE_API virtual void PerformAssetDiff(UObject* OldAsset, UObject* NewAsset, const FRevisionInfo& OldRevision, const FRevisionInfo& NewRevision) const override;
-
-private:
-	UE_API void ExecuteChangeDataAssetClass(TArray<TWeakObjectPtr<UDataAsset>> InDataAssets);
 };
 
 #undef UE_API
