SourceControl: Make 'Reloading...' a top level slow task dialog.

The slow tasks in `UPackageTools::ReloadPackage` were previously displayed as sub level actions under a top level Source Control action task (sync, revert, branch switch, branch merge). By making this a separate top level dialog it should become more clear to the user that the revision control operation is generally fast but the editor operation to reload packages is generally slow (most notably due to its inability for reloading OFPA files, which require an entire world reload).

Note: manually added SubmitTool tag because it failed to finish a NonUnity validator on this branch (it completed just fine on /Main) due to a path length Windows limitation.

#rnx

diff --git a/Engine/Source/Developer/AssetTools/Private/AssetViewUtils.cpp b/Engine/Source/Developer/AssetTools/Private/AssetViewUtils.cpp
--- a/Engine/Source/Developer/AssetTools/Private/AssetViewUtils.cpp
+++ b/Engine/Source/Developer/AssetTools/Private/AssetViewUtils.cpp
@@ -2,6 +2,7 @@
 
 #include "AssetViewUtils.h"
 #include "Algo/Transform.h"
+#include "Bookmarks/BookmarkScoped.h"
 #include "HAL/IConsoleManager.h"
 #include "Misc/MessageDialog.h"
 #include "HAL/FileManager.h"
@@ -1938,7 +1939,15 @@ bool AssetViewUtils::SyncPackagesFromSourceControl(const TArray<FString>& Packag
 		});
 
 		// Hot-reload the new packages...
-		UPackageTools::ReloadPackages(LoadedPackages);
+		if (LoadedPackages.Num() > 0)
+		{
+			FBookmarkScoped BookmarkScoped;
+			FScopedSlowTask SlowTaskScoped(1.f, NSLOCTEXT("SourceControl", "Reload", "Reloading..."));
+			SlowTaskScoped.MakeDialog();
+			SlowTaskScoped.EnterProgressFrame(1.f);
+
+			UPackageTools::ReloadPackages(LoadedPackages);
+		}
 
 		// Unload any deleted packages...
 		TArray<UPackage*> PackageRawPtrsToUnload;
@@ -2122,7 +2131,15 @@ static bool SyncPathsFromSourceControl(const FString& Revision, const TArray<FSt
 		}
 
 		// Hot-reload the new packages...
-		UPackageTools::ReloadPackages(LoadedPackages);
+		if (LoadedPackages.Num() > 0)
+		{
+			FBookmarkScoped BookmarkScoped;
+			FScopedSlowTask SlowTaskScoped(1.f, NSLOCTEXT("SourceControl", "Reload", "Reloading..."));
+			SlowTaskScoped.MakeDialog();
+			SlowTaskScoped.EnterProgressFrame(1.f);
+
+			UPackageTools::ReloadPackages(LoadedPackages);
+		}
 
 		// Unload any deleted packages...
 		TArray<UPackage*> PackageRawPtrsToUnload;
diff --git a/Engine/Source/Developer/SourceControl/Private/SourceControlHelpers.cpp b/Engine/Source/Developer/SourceControl/Private/SourceControlHelpers.cpp
--- a/Engine/Source/Developer/SourceControl/Private/SourceControlHelpers.cpp
+++ b/Engine/Source/Developer/SourceControl/Private/SourceControlHelpers.cpp
@@ -7,6 +7,7 @@
 #include "HAL/FileManager.h"
 #include "Misc/Paths.h"
 #include "Misc/ConfigContext.h"
+#include "Misc/ScopedSlowTask.h"
 #include "ISourceControlOperation.h"
 #include "SourceControlOperations.h"
 #include "ISourceControlProvider.h"
@@ -27,6 +28,7 @@
 #include "PackageTools.h"
 #include "ObjectTools.h"
 #include "FileHelpers.h"
+#include "Bookmarks/BookmarkScoped.h"
 #endif
 
 #define LOCTEXT_NAMESPACE "SourceControlHelpers"
@@ -1190,6 +1192,12 @@ bool USourceControlHelpers::ApplyOperationAndReloadPackages(const TArray<FString
 	// Hot-reload the new packages...
 	if (LoadedPackages.Num() > 0)
 	{
+		FBookmarkScoped BookmarkScoped;
+		FScopedSlowTask SlowTaskScoped(1.f, NSLOCTEXT("SourceControl", "Reload", "Reloading..."));
+		SlowTaskScoped.MakeDialog();
+		SlowTaskScoped.EnterProgressFrame(1.f);
+		SlowTaskScoped.ForceRefresh();
+
 		// Split into world and non-world packages and reload them separately.
 		TArray<UPackage*> LoadedWorldPackages;
 		TArray<UPackage*> LoadedNonWorldPackages;
diff --git a/Engine/Source/Editor/SceneOutliner/Private/SceneOutlinerSCCHandler.cpp b/Engine/Source/Editor/SceneOutliner/Private/SceneOutlinerSCCHandler.cpp
--- a/Engine/Source/Editor/SceneOutliner/Private/SceneOutlinerSCCHandler.cpp
+++ b/Engine/Source/Editor/SceneOutliner/Private/SceneOutlinerSCCHandler.cpp
@@ -12,7 +12,6 @@
 #include "AssetViewUtils.h"
 #include "LevelEditor.h"
 #include "RevisionControlStyle/RevisionControlStyle.h"
-#include "Bookmarks/BookmarkScoped.h"
 
 #define LOCTEXT_NAMESPACE "FSceneOutlinerSCCHandler"
 
@@ -451,11 +450,6 @@ void FSceneOutlinerSCCHandler::ExecuteSCCRevert()
 
 	if (PackagesToRevert.Num() > 0)
 	{
-		FBookmarkScoped BookmarkScoped;
-		FScopedSlowTask SlowTaskScoped(1.f, NSLOCTEXT("SourceControl", "RevertAll", "Reverting change(s)..."));
-		SlowTaskScoped.MakeDialog();
-		SlowTaskScoped.EnterProgressFrame(1.f);
-
 		SourceControlHelpers::RevertAndReloadPackages(PackagesToRevert, /*bRevertAll=*/false, /*bReloadWorld=*/true);
 	}
 }
diff --git a/Engine/Source/Editor/SourceControlWindows/Private/SSourceControlSubmit.cpp b/Engine/Source/Editor/SourceControlWindows/Private/SSourceControlSubmit.cpp
--- a/Engine/Source/Editor/SourceControlWindows/Private/SSourceControlSubmit.cpp
+++ b/Engine/Source/Editor/SourceControlWindows/Private/SSourceControlSubmit.cpp
@@ -30,7 +30,6 @@
 #include "Virtualization/VirtualizationSystem.h"
 #include "Logging/MessageLog.h"
 #include "RevisionControlStyle/RevisionControlStyle.h"
-#include "Bookmarks/BookmarkScoped.h"
 #include "HAL/IConsoleManager.h"
 #include "Algo/AllOf.h"
 #include "ToolMenu.h"
@@ -756,11 +755,6 @@ void SSourceControlSubmitWidget::OnRevert()
 	}
 
 	{
-		FBookmarkScoped BookmarkScoped;
-		FScopedSlowTask SlowTaskScoped(1.f, NSLOCTEXT("SourceControl", "RevertAll", "Reverting change(s)..."));
-		SlowTaskScoped.MakeDialog();
-		SlowTaskScoped.EnterProgressFrame(1.f);
-
 		bool bAnyReverted = false;
 		if (PackagesToRevert.Num() > 0)
 		{
diff --git a/Engine/Source/Editor/SourceControlWindows/Private/SourceControlWindows.cpp b/Engine/Source/Editor/SourceControlWindows/Private/SourceControlWindows.cpp
--- a/Engine/Source/Editor/SourceControlWindows/Private/SourceControlWindows.cpp
+++ b/Engine/Source/Editor/SourceControlWindows/Private/SourceControlWindows.cpp
@@ -17,7 +17,6 @@
 #include "Misc/ScopedSlowTask.h"
 #include "Widgets/Notifications/SNotificationList.h"
 #include "SourceControlSettings.h"
-#include "Bookmarks/BookmarkScoped.h"
 #include "HAL/IConsoleManager.h"
 
 #include UE_INLINE_GENERATED_CPP_BY_NAME(SourceControlWindows)
@@ -439,11 +438,6 @@ bool FSourceControlWindows::PromptForCheckin(FCheckinResultInfo& OutResultInfo,
 		const bool bSyncNeeded = FSourceControlWindows::CanSyncLatest();
 		if (bSyncNeeded)
 		{
-			FBookmarkScoped BookmarkScoped;
-			FScopedSlowTask SlowTaskScoped(1.f, NSLOCTEXT("SourceControl", "SyncLatest", "Syncing latest change(s)..."));
-			SlowTaskScoped.MakeDialog();
-			SlowTaskScoped.EnterProgressFrame(1.f);
-
 			if (!FSourceControlWindows::SyncLatest())
 			{
 				OutResultInfo.Description = LOCTEXT("SCC_Checkin_Aborted_Sync", "File check in aborted because the sync to the latest snapshot failed.");
diff --git a/Engine/Source/Editor/StatusBar/Private/SourceControlMenuHelpers.cpp b/Engine/Source/Editor/StatusBar/Private/SourceControlMenuHelpers.cpp
--- a/Engine/Source/Editor/StatusBar/Private/SourceControlMenuHelpers.cpp
+++ b/Engine/Source/Editor/StatusBar/Private/SourceControlMenuHelpers.cpp
@@ -25,7 +25,6 @@
 #include "Misc/MessageDialog.h"
 #include "Misc/ScopedSlowTask.h"
 #include "RevisionControlStyle/RevisionControlStyle.h"
-#include "Bookmarks/BookmarkScoped.h"
 #include "Styling/StyleColors.h"
 #include "HAL/IConsoleManager.h"
 #include "SSourceControlControls.h"
@@ -279,11 +278,6 @@ void FSourceControlCommands::RevertAllModifiedFiles_Clicked()
 		const bool bCanBeDeclined = false;
 		FEditorFileUtils::SaveDirtyPackages(bPromptUserToSave, bSaveMapPackages, bSaveContentPackages, bFastSave, bNotifyNoPackagesSaved, bCanBeDeclined);
 
-		FBookmarkScoped BookmarkScoped;
-		FScopedSlowTask SlowTaskScoped(1.f, NSLOCTEXT("SourceControl", "RevertAll", "Reverting change(s)..."));
-		SlowTaskScoped.MakeDialog();
-		SlowTaskScoped.EnterProgressFrame(1.f);
-
 		FSourceControlWindows::RevertAllChangesAndReloadWorld();
 	}
 }
