Programmatically filter migration of dependencies on Migrate function

#Synced-CL 50374992

diff --git a/Engine/Source/Developer/AssetTools/Private/AssetTools.cpp b/Engine/Source/Developer/AssetTools/Private/AssetTools.cpp
--- a/Engine/Source/Developer/AssetTools/Private/AssetTools.cpp
+++ b/Engine/Source/Developer/AssetTools/Private/AssetTools.cpp
@@ -3437,6 +3437,11 @@ void UAssetToolsImpl::MigratePackages(const TArray<FName>& PackageNamesToMigrate
 }
 
 void UAssetToolsImpl::MigratePackages(const TArray<FName>& PackageNamesToMigrate, const FString& DestinationPath, const struct FMigrationOptions& Options) const
+{
+	MigratePackages(PackageNamesToMigrate, DestinationPath, TSet<FName>(), TSet<FName>(), Options);
+}
+
+void UAssetToolsImpl::MigratePackages(const TArray<FName>& PackageNamesToMigrate, const FString& DestinationPath, TSet<FName> AllDependenciesToMigrate, TSet<FName> ExcludedDependencies, const struct FMigrationOptions& Options) const
 {
 	// Packages must be saved for the migration to work
 	const bool bPromptUserToSave = CanShowDialogs() && Options.bPrompt;
@@ -3444,18 +3449,29 @@ void UAssetToolsImpl::MigratePackages(const TArray<FName>& PackageNamesToMigrate
 	const bool bSaveContentPackages = true;
 	if (FEditorFileUtils::SaveDirtyPackages(bPromptUserToSave, bSaveMapPackages, bSaveContentPackages))
 	{
+		auto RunMigratePackages = [this, PackageNamesToMigrate, DestinationPath, Options, AllDependenciesToMigrate = MoveTemp(AllDependenciesToMigrate), ExcludedDependencies = MoveTemp(ExcludedDependencies)]() mutable
+			{
+				bool bSanitizeInputs = true;
+				if (AllDependenciesToMigrate.IsEmpty() && ExcludedDependencies.IsEmpty())
+				{
+					BuildMigrateDependencies(PackageNamesToMigrate, AllDependenciesToMigrate, ExcludedDependencies);
+					bSanitizeInputs = false;
+				}
+				PerformMigratePackages(PackageNamesToMigrate, DestinationPath, Options, AllDependenciesToMigrate, ExcludedDependencies, bSanitizeInputs);
+			};
+		
 		FAssetRegistryModule& AssetRegistryModule = FModuleManager::Get().LoadModuleChecked<FAssetRegistryModule>(TEXT("AssetRegistry"));
 		if (AssetRegistryModule.Get().IsLoadingAssets())
 		{
 			// Open a dialog asking the user to wait while assets are being discovered
 			SDiscoveringAssetsDialog::OpenDiscoveringAssetsDialog(
-				SDiscoveringAssetsDialog::FOnAssetsDiscovered::CreateUObject(this, &UAssetToolsImpl::PerformMigratePackages, PackageNamesToMigrate, DestinationPath, Options)
+				SDiscoveringAssetsDialog::FOnAssetsDiscovered::CreateLambda(RunMigratePackages)
 			);
 		}
 		else
 		{
 			// Assets are already discovered, perform the migration now
-			PerformMigratePackages(PackageNamesToMigrate, DestinationPath, Options);
+			RunMigratePackages();
 		}
 	}
 }
@@ -4906,12 +4922,12 @@ bool UAssetToolsImpl::CanCreateAsset(const FString& AssetName, const FString& Pa
 	return true;
 }
 
-void UAssetToolsImpl::PerformMigratePackages(TArray<FName> PackageNamesToMigrate, const FString DestinationPath, const FMigrationOptions Options) const
+void UAssetToolsImpl::BuildMigrateDependencies(const TArray<FName>& PackageNamesToMigrate, TSet<FName>& OutAllDependenciesToMigrate, TSet<FName>& OutExcludedDependencies) const
 {
 	// Form a full list of packages to move by including the dependencies of the supplied packages
-	TSet<FName> AllPackageNamesToMove;
+	OutAllDependenciesToMigrate.Reset();
 	TSet<FString> ExternalObjectsPaths;
-	TSet<FName> ExcludedDependencies;
+	OutExcludedDependencies.Reset();
 	{
 		FScopedSlowTask SlowTask(static_cast<float>(PackageNamesToMigrate.Num()), LOCTEXT( "MigratePackages_GatheringDependencies", "Gathering Dependencies..." ) );
 		SlowTask.MakeDialog();
@@ -4922,14 +4938,42 @@ void UAssetToolsImpl::PerformMigratePackages(TArray<FName> PackageNamesToMigrate
 		{
 			SlowTask.EnterProgressFrame();
 
-			if ( !AllPackageNamesToMove.Contains(*PackageIt) )
+			if (!OutAllDependenciesToMigrate.Contains(*PackageIt) )
 			{
-				AllPackageNamesToMove.Add(*PackageIt);
-				RecursiveGetDependencies(*PackageIt, AllPackageNamesToMove, ExternalObjectsPaths, ExcludedDependencies, ShouldExcludePackage);
+				OutAllDependenciesToMigrate.Add(*PackageIt);
+				RecursiveGetDependencies(*PackageIt, OutAllDependenciesToMigrate, ExternalObjectsPaths, OutExcludedDependencies, ShouldExcludePackage);
 			}
 		}
 	}
+}
 
+void UAssetToolsImpl::PerformMigratePackages(const TArray<FName>& PackageNamesToMigrate, const FString& DestinationPath, const FMigrationOptions& Options, TSet<FName> AllDependenciesToMigrate, TSet<FName> ExcludedDependencies, bool bSanitizeInputs) const
+{
+	if (bSanitizeInputs)
+	{
+		TSet<FName> VerifiedAllPackageNamesToMove;
+		TSet<FName> VerifiedExcludedDependencies;
+		BuildMigrateDependencies(PackageNamesToMigrate, VerifiedAllPackageNamesToMove, VerifiedExcludedDependencies);
+		
+		// Don't allow migration of packages that are not in the allowed list of verified packages 
+		for (auto It = AllDependenciesToMigrate.CreateIterator(); It; ++It)
+		{
+			if (!VerifiedAllPackageNamesToMove.Contains(*It))
+			{
+				It.RemoveCurrent();
+			}
+		}
+		
+		// Don't allow migration of dependencies that have been excluded in verified list
+		for (const FName& Ex : VerifiedExcludedDependencies)
+		{
+			if (!ExcludedDependencies.Contains(Ex))
+			{
+				ExcludedDependencies.Add(Ex);
+			}
+		}
+	}
+	
 	// Fetch the enabled plugins and their mount points
 	TMap<FName, EPluginLoadedFrom> EnabledPluginToLoadedFrom;
 	TArray<TSharedRef<IPlugin>> EnabledPlugins = IPluginManager::Get().GetEnabledPluginsWithContent();
@@ -4942,7 +4986,7 @@ void UAssetToolsImpl::PerformMigratePackages(TArray<FName> PackageNamesToMigrate
 	TSet<FName> ShouldMigratePackage;
 	bool bShouldShowEngineContent = GetDefault<UContentBrowserSettings>()->GetDisplayEngineFolder();
 	TMap<FName, TSet<FName>> PackageToExternalObjectPackages;
-		
+	TMap<FName, FName> ExternalObjectPackageToParentPackage;
 	{
 		FAssetRegistryModule& AssetRegistryModule = FModuleManager::Get().LoadModuleChecked<FAssetRegistryModule>(TEXT("AssetRegistry"));
 		const bool bOnlyIncludeOnDiskAssets = true;
@@ -4950,7 +4994,7 @@ void UAssetToolsImpl::PerformMigratePackages(TArray<FName> PackageNamesToMigrate
 		// This is the new list to prompt for migration
 		TSet<FName> FilteredPackageNamesToMove;
 
-		for (const FName& PackageName : AllPackageNamesToMove)
+		for (const FName& PackageName : AllDependenciesToMigrate)
 		{
 			// Associate External Packages to Level Packages
 			FString PackageNameStr = PackageName.ToString();
@@ -4963,7 +5007,13 @@ void UAssetToolsImpl::PerformMigratePackages(TArray<FName> PackageNamesToMigrate
 					{
 						if (!AssetData.GetOptionalOuterPathName().IsNone())
 						{
-							PackageToExternalObjectPackages.FindOrAdd(FSoftObjectPath(AssetData.GetOptionalOuterPathName().ToString()).GetLongPackageFName()).Add(PackageName);
+							const FName OuterPackageName = FSoftObjectPath(AssetData.GetOptionalOuterPathName().ToString()).GetLongPackageFName();
+							PackageToExternalObjectPackages.FindOrAdd(OuterPackageName).Add(PackageName);
+							
+							if (bSanitizeInputs)
+							{
+								ExternalObjectPackageToParentPackage.Add(PackageName, OuterPackageName);
+							}
 						}
 					}
 				}
@@ -5039,11 +5089,11 @@ void UAssetToolsImpl::PerformMigratePackages(TArray<FName> PackageNamesToMigrate
 			}
 		}
 
-		AllPackageNamesToMove = FilteredPackageNamesToMove;
+		AllDependenciesToMigrate = FilteredPackageNamesToMove;
 	}
 
 	// Confirm that there is at least one package to move 
-	if (AllPackageNamesToMove.Num() == 0)
+	if (AllDependenciesToMigrate.Num() == 0)
 	{
 		if (CanShowDialogs())
 		{
@@ -5052,12 +5102,40 @@ void UAssetToolsImpl::PerformMigratePackages(TArray<FName> PackageNamesToMigrate
 		return;
 	}
 
+	TSet<FName> SelectedParents;
 	TSharedPtr<TArray<ReportPackageData>> ReportPackages = MakeShareable(new TArray<ReportPackageData>);
-	for (auto PackageIt = AllPackageNamesToMove.CreateConstIterator(); PackageIt; ++PackageIt)
+	for (auto PackageIt = AllDependenciesToMigrate.CreateConstIterator(); PackageIt; ++PackageIt)
 	{
-		bool bShouldMigratePackage = ShouldMigratePackage.Find(*PackageIt) != nullptr;
-		ReportPackages.Get()->Add({ (*PackageIt).ToString(), bShouldMigratePackage });
+		const FName PackageName = *PackageIt;
+		bool bShouldMigratePackage = ShouldMigratePackage.Find(PackageName) != nullptr;
+		
+		if (bSanitizeInputs && bShouldMigratePackage)
+		{
+			// For each external object/actor if their parent is not being migrated, then they shouldn't either
+			if (const FName* ParentPackageName = ExternalObjectPackageToParentPackage.Find(PackageName))
+			{
+				if (!ShouldMigratePackage.Contains(*ParentPackageName))
+				{
+					bShouldMigratePackage = false;
+				}
+			}
+		}
+		
+		ReportPackages.Get()->Add({ PackageName.ToString(), bShouldMigratePackage });
 	}
+	
+	if (bSanitizeInputs)
+	{
+		// Clean PackageToExternalObjectPackages to keep only the selected Package parents
+		for (auto It = PackageToExternalObjectPackages.CreateIterator(); It; ++It)
+		{
+			if (!ShouldMigratePackage.Contains(It.Key()))
+			{
+				It.RemoveCurrent();
+			}
+		}
+	}
+	
 	// Prompt the user displaying all assets that are going to be migrated
 	if(CanShowDialogs() && Options.bPrompt)
 	{
diff --git a/Engine/Source/Developer/AssetTools/Private/AssetTools.h b/Engine/Source/Developer/AssetTools/Private/AssetTools.h
--- a/Engine/Source/Developer/AssetTools/Private/AssetTools.h
+++ b/Engine/Source/Developer/AssetTools/Private/AssetTools.h
@@ -124,6 +124,8 @@ public:
 	virtual bool CreateDiffProcess(const FString& DiffCommand, const FString& OldTextFilename, const FString& NewTextFilename, const FString& DiffArgs = FString("")) const override;
 	virtual void MigratePackages(const TArray<FName>& PackageNamesToMigrate) const override;
 	virtual void MigratePackages(const TArray<FName>& PackageNamesToMigrate, const FString& TargetPath, const struct FMigrationOptions& Options = FMigrationOptions()) const override;
+	virtual void MigratePackages(const TArray<FName>& PackageNamesToMigrate, const FString& TargetPath, TSet<FName> AllDependenciesToMigrate, TSet<FName> ExcludedDependencies, const struct FMigrationOptions& Options = FMigrationOptions()) const override;
+	virtual void BuildMigrateDependencies(const TArray<FName>& PackageNamesToMigrate, TSet<FName>& OutAllDependenciesToMigrate, TSet<FName>& OutExcludedDependencies) const override;
 	virtual UE::AssetTools::FOnPackageMigration& GetOnPackageMigration() override;
 	virtual void BeginAdvancedCopyPackages(const TArray<FName>& InputNamesToCopy, const FString& TargetPath) const override;
 	virtual void BeginAdvancedCopyPackages(const TArray<FName>& InputNamesToCopy, const FString& TargetPath, const FAdvancedCopyCompletedEvent& OnCopyComplete) const override;
@@ -215,7 +217,7 @@ private:
 	bool CanCreateAsset(const FString& AssetName, const FString& PackageName, const FText& OperationText) const;
 
 	/** Begins the package migration, after assets have been discovered */
-	void PerformMigratePackages(TArray<FName> PackageNamesToMigrate, const FString DestinationPath, const FMigrationOptions Options) const;
+	void PerformMigratePackages(const TArray<FName>& PackageNamesToMigrate, const FString& DestinationPath, const FMigrationOptions& Options, TSet<FName> AllDependenciesToMigrate, TSet<FName> ExcludedDependencies, bool bSanitizeInputs) const;
 
 	/** Check if an asset can migrated */
 	bool CanMigratePackage(FName PackageName) const;
diff --git a/Engine/Source/Developer/AssetTools/Public/IAssetTools.h b/Engine/Source/Developer/AssetTools/Public/IAssetTools.h
--- a/Engine/Source/Developer/AssetTools/Public/IAssetTools.h
+++ b/Engine/Source/Developer/AssetTools/Public/IAssetTools.h
@@ -629,7 +629,13 @@ public:
 	/* Migrate packages and dependencies to another folder */
 	UFUNCTION(BlueprintCallable, Category = "Editor Scripting | Asset Tools", BlueprintPure = false)
 	virtual void MigratePackages(const TArray<FName>& PackageNamesToMigrate, const FString& DestinationPath, const struct FMigrationOptions& Options = FMigrationOptions()) const = 0;
+	
+	/* Migrate packages and dependencies to another folder */
+	virtual void MigratePackages(const TArray<FName>& PackageNamesToMigrate, const FString& DestinationPath, TSet<FName> AllDependenciesToMigrate, TSet<FName> ExcludedDependencies, const struct FMigrationOptions& Options = FMigrationOptions()) const = 0;
 
+	/* Build dependency packages and excluded dependencies from a list of packages to migrate */
+	virtual void BuildMigrateDependencies(const TArray<FName>& PackageNamesToMigrate, TSet<FName>& OutAllDependenciesToMigrate, TSet<FName>& OutExcludedDependencies) const = 0;
+	
 	/**
 	 * Event called when some packages are migrated
 	 * Note this is only true when AssetTools.UseNewPackageMigration is true
