PR #5783: Add Editor Settings to configure default name with prefix and suffix for new assets

#Synced-CL 50404663

diff --git a/Engine/Plugins/Animation/ControlRig/Source/ControlRigEditor/Private/ControlRigAssetActions.cpp b/Engine/Plugins/Animation/ControlRig/Source/ControlRigEditor/Private/ControlRigAssetActions.cpp
--- a/Engine/Plugins/Animation/ControlRig/Source/ControlRigEditor/Private/ControlRigAssetActions.cpp
+++ b/Engine/Plugins/Animation/ControlRig/Source/ControlRigEditor/Private/ControlRigAssetActions.cpp
@@ -286,7 +286,7 @@ UControlRigRuntimeAsset* FControlRigAssetActions::CreateControlRigFromSkeletalMe
 	FString PackagePath = InSelectedObject->GetPathName();
 
 	FAssetToolsModule& AssetToolsModule = FModuleManager::LoadModuleChecked<FAssetToolsModule>("AssetTools");
-	const TOptional<FString> DefaultControlRigName = AssetToolsModule.Get().GetDefaultAssetNameForClassFromSource(UControlRigRuntimeAsset::StaticClass(), InSelectedObject);
+	const TOptional<FString> DefaultControlRigName = AssetToolsModule.Get().GetDefaultAssetNameForClass(UControlRigRuntimeAsset::StaticClass(), nullptr, InSelectedObject);
 
 	FString ControlRigName = DefaultControlRigName.IsSet() ? DefaultControlRigName.GetValue() : FString::Printf(TEXT("%s_CtrlRig"), *InSelectedObject->GetName());
 
diff --git a/Engine/Plugins/Animation/ControlRig/Source/ControlRigEditor/Private/ControlRigBlueprintActions.cpp b/Engine/Plugins/Animation/ControlRig/Source/ControlRigEditor/Private/ControlRigBlueprintActions.cpp
--- a/Engine/Plugins/Animation/ControlRig/Source/ControlRigEditor/Private/ControlRigBlueprintActions.cpp
+++ b/Engine/Plugins/Animation/ControlRig/Source/ControlRigEditor/Private/ControlRigBlueprintActions.cpp
@@ -267,9 +267,9 @@ UControlRigBlueprint* FControlRigBlueprintActions::CreateControlRigFromSkeletalM
 	FString PackagePath = InSelectedObject->GetPathName();
 
 	FAssetToolsModule& AssetToolsModule = FModuleManager::LoadModuleChecked<FAssetToolsModule>("AssetTools");
-	const TOptional<FString> DefaultControlRigName = AssetToolsModule.Get().GetDefaultAssetNameForClassFromSource(UControlRigBlueprint::StaticClass(), InSelectedObject);
+	const TOptional<FString> DefaultControlRigName = AssetToolsModule.Get().GetDefaultAssetNameForClass(UControlRigBlueprint::StaticClass(), nullptr, InSelectedObject);
 
-	FString ControlRigName = DefaultControlRigName.IsSet() ? DefaultControlRigName.GetValue() : FString::Printf(TEXT("%s_CtrlRig"), *InSelectedObject->GetName());
+	FString ControlRigName = DefaultControlRigName.Get(FString::Printf(TEXT("%s_CtrlRig"), *InSelectedObject->GetName()));
 
 	int32 LastSlashPos = INDEX_NONE;
 	if (PackagePath.FindLastChar('/', LastSlashPos))
diff --git a/Engine/Plugins/Editor/ContentBrowser/ContentBrowserAssetDataSource/Source/ContentBrowserAssetDataSource/Private/ContentBrowserAssetDataSource.cpp b/Engine/Plugins/Editor/ContentBrowser/ContentBrowserAssetDataSource/Source/ContentBrowserAssetDataSource/Private/ContentBrowserAssetDataSource.cpp
--- a/Engine/Plugins/Editor/ContentBrowser/ContentBrowserAssetDataSource/Source/ContentBrowserAssetDataSource/Private/ContentBrowserAssetDataSource.cpp
+++ b/Engine/Plugins/Editor/ContentBrowser/ContentBrowserAssetDataSource/Source/ContentBrowserAssetDataSource/Private/ContentBrowserAssetDataSource.cpp
@@ -3986,8 +3986,8 @@ void UContentBrowserAssetDataSource::OnNewAssetRequested(const FName InPath, TWe
 		{
 			FEditorDelegates::OnNewAssetCreated.Broadcast(NewFactory);
 
-			const TOptional<FString> DefaultAssetName = AssetTools->GetDefaultAssetNameForClass(NewFactory->GetSupportedClass());
-			const FString BaseAssetName = DefaultAssetName.IsSet() ? DefaultAssetName.GetValue() : NewFactory->GetDefaultNewAssetName();
+			const TOptional<FString> DefaultAssetName = AssetTools->GetDefaultAssetNameForClass(NewFactory->GetSupportedClass(), NewFactory);
+			const FString BaseAssetName = DefaultAssetName.Get(NewFactory->GetDefaultNewAssetName());
 
 			FString UniqueAssetName;
 			FString PackageNameToUse;
diff --git a/Engine/Source/Developer/AssetTools/Private/AssetTools.cpp b/Engine/Source/Developer/AssetTools/Private/AssetTools.cpp
--- a/Engine/Source/Developer/AssetTools/Private/AssetTools.cpp
+++ b/Engine/Source/Developer/AssetTools/Private/AssetTools.cpp
@@ -1841,9 +1841,9 @@ void UAssetToolsImpl::CreateAssetsFrom(TConstArrayView<UObject*> SourceObjects,
 			if (UFactory* Factory = FactoryConstructor(SourceObject))
 			{
 				// Create an appropriate and unique name 
-				TOptional<FString> DefaultAssetName = GetDefaultAssetNameForClassFromSource(CreateAssetType, SourceObject);
+				TOptional<FString> DefaultAssetName = GetDefaultAssetNameForClass(CreateAssetType, Factory, SourceObject);
 
-				const FString BasePackageName = DefaultAssetName.IsSet() ? DefaultAssetName.GetValue() : SourceObject->GetName();
+				const FString BasePackageName = DefaultAssetName.Get(SourceObject->GetName());
 				const FString Suffix = DefaultAssetName.IsSet() ? FString() : DefaultSuffix;
 
 				const FString PackagePath = FPackageName::GetLongPackagePath(SourceObject->GetPackage()->GetName());
@@ -1868,9 +1868,9 @@ void UAssetToolsImpl::CreateAssetsFrom(TConstArrayView<UObject*> SourceObjects,
 				if (UFactory* Factory = FactoryConstructor(SourceObject))
 				{
 					// Determine an appropriate name
-					const TOptional<FString> DefaultAssetName = GetDefaultAssetNameForClassFromSource(CreateAssetType, SourceObject);
+					TOptional<FString> DefaultAssetName = GetDefaultAssetNameForClass(CreateAssetType, Factory, SourceObject);
 
-					const FString BasePackageName = DefaultAssetName.IsSet() ? DefaultAssetName.GetValue() : SourceObject->GetName();
+					const FString BasePackageName = DefaultAssetName.Get(SourceObject->GetName());
 					const FString Suffix = DefaultAssetName.IsSet() ? FString() : DefaultSuffix;
 
 					const FString PackagePath = FPackageName::GetLongPackagePath(SourceObject->GetPackage()->GetName());
@@ -1936,8 +1936,8 @@ UObject* UAssetToolsImpl::CreateAssetWithDialog(UClass* AssetClass, UFactory* Fa
 		}
 
 		// Determine an appropriate name
-		const TOptional<FString> DefaultAssetName = GetDefaultAssetNameForClass(AssetClass);
-		const FString BaseAssetName = DefaultAssetName.IsSet() ? DefaultAssetName.GetValue() : Factory->GetDefaultNewAssetName();
+		TOptional<FString> DefaultAssetName = GetDefaultAssetNameForClass(AssetClass, Factory, nullptr);
+		const FString BaseAssetName = DefaultAssetName.Get(Factory->GetDefaultNewAssetName());
 
 		FString PackageName;
 		FString AssetName;
@@ -2021,9 +2021,13 @@ bool UAssetToolsImpl::CreateAssetWithDialogAsync(UClass* AssetClass, UFactory* F
 		AssetPath = TEXT("/Game");
 	}
 
+	// Determine an appropriate name
+	const TOptional<FString> DefaultAssetName = GetDefaultAssetNameForClass(AssetClass, Factory, nullptr);
+	const FString BaseAssetName = DefaultAssetName.Get(Factory->GetDefaultNewAssetName());
+
 	FString PackageName;
 	FString AssetName;
-	CreateUniqueAssetName(AssetPath / Factory->GetDefaultNewAssetName(), TEXT(""), PackageName, AssetName);
+	CreateUniqueAssetName(AssetPath / BaseAssetName, TEXT(""), PackageName, AssetName);
 
 	return CreateAssetWithDialogAsync(AssetName, AssetPath, AssetClass, Factory, OnComplete, OnCancelled, CallingContext);
 }
@@ -7160,50 +7164,43 @@ void UAssetToolsImpl::UnregisterSanitizeNameDelegate(const FName OwnerName)
 	SanitizeNameDelegates.Remove(OwnerName);
 }
 
-TOptional<FString> UAssetToolsImpl::GetDefaultAssetNameForClass(const UClass* Class)
+TOptional<FString> UAssetToolsImpl::GetDefaultAssetNameForClass(const UClass* Class, const UFactory* Factory, const UObject* SourceObject)
 {
+	FString DefaultAssetName;
 	if (FString* AssetName = DefaultNewAssetNames.Find(Class))
 	{
-		FString DefaultAssetName = *AssetName;
-		SanitizeName(DefaultAssetName);
-
-		FString SanitizedDefaultName = ObjectTools::SanitizeInvalidChars(DefaultAssetName, INVALID_LONGPACKAGE_CHARACTERS INVALID_OBJECTNAME_CHARACTERS);
-
-		FText ErrorText;
-		if (!FFileHelper::IsFilenameValidForSaving(SanitizedDefaultName, ErrorText))
-		{
-			UE_LOG(LogAssetTools, Warning, TEXT("The default name '%s' for class '%s' is invalid: %s"), *SanitizedDefaultName, *Class->GetName(), *ErrorText.ToString());
-			return TOptional<FString>();
-		}
-
-		return SanitizedDefaultName;
+		DefaultAssetName = *AssetName;
 	}
 
-	return TOptional<FString>();
-}
+	if (const FAssetNamingConventionEntry* NamingConvention = GetDefault<UAssetToolsEditorSettings>()->AssetNamingConventionOverrides.Find(Class))
+	{
+		DefaultAssetName = FString::Format(TEXT("{0}{1}{2}"),
+			{
+				NamingConvention->Prefix.Get(TEXT("")),
+				DefaultAssetName.IsEmpty() ? NamingConvention->AssetName.Get(Factory ? Factory->GetDefaultNewAssetName() : TEXT("")) : DefaultAssetName,
+				NamingConvention->Suffix.Get(TEXT(""))
+			});
+	}
 
-TOptional<FString> UAssetToolsImpl::GetDefaultAssetNameForClassFromSource(const UClass* Class, const UObject* SourceObject)
-{
-	if (FString* AssetName = DefaultNewAssetNames.Find(Class))
+	if (DefaultAssetName.IsEmpty())
 	{
-		FString DefaultAssetName = *AssetName;
-		DefaultAssetName = DefaultAssetName.Replace(TEXT("{source}"), *SourceObject->GetName(), ESearchCase::IgnoreCase);
+		return TOptional<FString>();
+	}
 
-		SanitizeName(DefaultAssetName);
+	DefaultAssetName = DefaultAssetName.Replace(TEXT("{source}"), SourceObject ? *SourceObject->GetName() : TEXT(""), ESearchCase::IgnoreCase);
 
-		FString SanitizedDefaultName = ObjectTools::SanitizeInvalidChars(DefaultAssetName, INVALID_LONGPACKAGE_CHARACTERS INVALID_OBJECTNAME_CHARACTERS);
+	SanitizeName(DefaultAssetName);
 
-		FText ErrorText;
-		if (!FFileHelper::IsFilenameValidForSaving(SanitizedDefaultName, ErrorText))
-		{
-			UE_LOG(LogAssetTools, Warning, TEXT("The default name '%s' for class '%s' is invalid: %s"), *SanitizedDefaultName, *Class->GetName(), *ErrorText.ToString());
-			return TOptional<FString>();
-		}
+	FString SanitizedDefaultName = ObjectTools::SanitizeInvalidChars(DefaultAssetName, INVALID_LONGPACKAGE_CHARACTERS INVALID_OBJECTNAME_CHARACTERS);
 
-		return SanitizedDefaultName;
+	FText ErrorText;
+	if (!FFileHelper::IsFilenameValidForSaving(SanitizedDefaultName, ErrorText))
+	{
+		UE_LOG(LogAssetTools, Warning, TEXT("The default name '%s' for class '%s' is invalid: %s"), *SanitizedDefaultName, *Class->GetName(), *ErrorText.ToString());
+		return TOptional<FString>();
 	}
 
-	return TOptional<FString>();
+	return SanitizedDefaultName;
 }
 
 void UAssetToolsImpl::RegisterDefaultAssetNameForClass(const UClass* Class, const FString& DefaultName)
diff --git a/Engine/Source/Developer/AssetTools/Private/AssetTools.h b/Engine/Source/Developer/AssetTools/Private/AssetTools.h
--- a/Engine/Source/Developer/AssetTools/Private/AssetTools.h
+++ b/Engine/Source/Developer/AssetTools/Private/AssetTools.h
@@ -179,8 +179,7 @@ public:
 	virtual bool SanitizeName(FString& NameToSanitize) override;
 	virtual void RegisterSanitizeNameDelegate(const FName OwnerName, FSanitizeName Delegate) override;
 	virtual void UnregisterSanitizeNameDelegate(const FName OwnerName) override;
-	virtual TOptional<FString> GetDefaultAssetNameForClass(const UClass* Class) override;
-	virtual TOptional<FString> GetDefaultAssetNameForClassFromSource(const UClass* Class, const UObject* SourceObject) override;
+	virtual TOptional<FString> GetDefaultAssetNameForClass(const UClass* Class, const UFactory* Factory, const UObject* SourceObject) override;
 	virtual void RegisterDefaultAssetNameForClass(const UClass* Class, const FString& DefaultName) override;
 	virtual void UnregisterDefaultAssetNameForClass(const UClass* Class) override;
 	virtual void RegisterCanMigrateAsset(const FName OwnerName, UE::AssetTools::FCanMigrateAsset Delegate) override;
diff --git a/Engine/Source/Developer/AssetTools/Public/AssetToolsSettings.h b/Engine/Source/Developer/AssetTools/Public/AssetToolsSettings.h
--- a/Engine/Source/Developer/AssetTools/Public/AssetToolsSettings.h
+++ b/Engine/Source/Developer/AssetTools/Public/AssetToolsSettings.h
@@ -36,3 +36,46 @@ public:
 	UPROPERTY(config, EditAnywhere, Category = "Advanced Copy", Meta = (TitleProperty = "ClassToCopy"))
 	TArray<FAdvancedCopyMap> AdvancedCopyCustomizations;
 };
+
+USTRUCT()
+struct FAssetNamingConventionEntry
+{
+	GENERATED_BODY()
+
+	/**
+	 * An optional prefix to include when naming an asset
+	 */
+	UPROPERTY(EditAnywhere, Category = "Asset Naming")
+	TOptional<FString> Prefix;
+
+	/* Supports {source} when used in the context of creating assets from existing uobjects instead of purely from factory
+	 * For instance, to override "NewBlueprint" with "BP_", "BP_{source}_" to preserve existing source object names,
+	 * or "BP_" to always reset the link
+	 * 
+	 * Note - if using a prefix or suffix, you must also set this (todo this is terrible)
+	 */
+	UPROPERTY(EditAnywhere, Category = "Asset Naming")
+	TOptional<FString> AssetName;
+
+	/**
+	 * An optional suffix to include when naming an asset
+	 */
+	UPROPERTY(EditAnywhere, Category = "Asset Naming")
+	TOptional<FString> Suffix;
+};
+
+UCLASS(config = EditorSettings, defaultconfig, meta = (DisplayName = "Asset Tools"))
+class UAssetToolsEditorSettings : public UDeveloperSettings
+{
+	GENERATED_BODY()
+
+public:
+	UAssetToolsEditorSettings() {};
+
+	/**
+	 * Override default name for new assets, by asset type.
+	 * Use of an override will override any factory setting
+	 */
+	UPROPERTY(EditAnywhere, config, Category = "Asset Naming")
+	TMap<TSoftClassPtr<UObject>, FAssetNamingConventionEntry> AssetNamingConventionOverrides;
+};
diff --git a/Engine/Source/Developer/AssetTools/Public/IAssetTools.h b/Engine/Source/Developer/AssetTools/Public/IAssetTools.h
--- a/Engine/Source/Developer/AssetTools/Public/IAssetTools.h
+++ b/Engine/Source/Developer/AssetTools/Public/IAssetTools.h
@@ -849,9 +849,14 @@ public:
 	virtual void UnregisterSanitizeNameDelegate(const FName OwnerName) = 0;
 
 	/** Get the default name for the supplied class (if one is registered) to use when creating new assets of that type */
-	virtual TOptional<FString> GetDefaultAssetNameForClass(const UClass* Class) = 0;
+	virtual TOptional<FString> GetDefaultAssetNameForClass(const UClass* Class, const UFactory* Factory = nullptr, const UObject* SourceObject = nullptr) = 0;
+
 	/** Get the default name for the supplied class (if one is registered) to use when creating a new asset of that type from the input source object */
-	virtual TOptional<FString> GetDefaultAssetNameForClassFromSource(const UClass* Class, const UObject* SourceObject) = 0;
+	UE_DEPRECATED(5.8, "Use GetDefaultAssetNameForClass and pass a source object")
+	virtual TOptional<FString> GetDefaultAssetNameForClassFromSource(const UClass* Class, const UObject* SourceObject, const UFactory* Factory = nullptr) final
+	{
+		return GetDefaultAssetNameForClass(Class, Factory, SourceObject);
+	}
 	/** Register a default name for the supplied class to use when creating new assets of that type */
 	virtual void RegisterDefaultAssetNameForClass(const UClass* Class, const FString& DefaultName) = 0;
 	/** Remove any previously registered name for the supplied class */
