- Added a bOverwriteExisting option to CreateAsset, useful when automating/script object creation to not prompt

#Synced-CL 50913554

diff --git a/Engine/Source/Developer/AssetTools/Private/AssetTools.cpp b/Engine/Source/Developer/AssetTools/Private/AssetTools.cpp
--- a/Engine/Source/Developer/AssetTools/Private/AssetTools.cpp
+++ b/Engine/Source/Developer/AssetTools/Private/AssetTools.cpp
@@ -1715,7 +1715,7 @@ TWeakPtr<IClassTypeActions> UAssetToolsImpl::GetClassTypeActionsForClass( UClass
 	return MostDerivedClassTypeActions;
 }
 
-UObject* UAssetToolsImpl::CreateAsset(const FString& AssetName, const FString& PackagePath, UClass* AssetClass, UFactory* Factory, FName CallingContext)
+UObject* UAssetToolsImpl::CreateAsset(const FString& AssetName, const FString& PackagePath, UClass* AssetClass, UFactory* Factory, FName CallingContext, bool bOverwriteExisting)
 {
 	FGCObjectScopeGuard DontGCFactory(Factory);
 
@@ -1745,7 +1745,7 @@ UObject* UAssetToolsImpl::CreateAsset(const FString& AssetName, const FString& P
 	const FString PackageName = UPackageTools::SanitizePackageName(PackagePath + TEXT("/") + AssetName);
 
 	// Make sure we can create the asset without conflicts
-	if ( !CanCreateAsset(AssetName, PackageName, LOCTEXT("CreateANewObject", "Create a new object")) )
+	if ( !CanCreateAsset(AssetName, PackageName, LOCTEXT("CreateANewObject", "Create a new object"), bOverwriteExisting) )
 	{
 		return nullptr;
 	}
@@ -1781,7 +1781,7 @@ UObject* UAssetToolsImpl::CreateAsset(const FString& AssetName, const FString& P
 	return NewObj;
 }
 
-bool UAssetToolsImpl::CreateAssetAsync(const FString& AssetName, const FString& PackagePath, UClass* AssetClass, UFactory* Factory, FAssetCreateComplete OnComplete, FAssetCreateCancelled OnCancelled, FName CallingContext)
+bool UAssetToolsImpl::CreateAssetAsync(const FString& AssetName, const FString& PackagePath, UClass* AssetClass, UFactory* Factory, FAssetCreateComplete OnComplete, FAssetCreateCancelled OnCancelled, FName CallingContext, bool bOverwriteExisting)
 {
 	// Verify the factory class
 	if (!ensure(AssetClass && Factory))
@@ -1810,7 +1810,7 @@ bool UAssetToolsImpl::CreateAssetAsync(const FString& AssetName, const FString&
 	const FString PackageName = UPackageTools::SanitizePackageName(PackagePath + TEXT("/") + AssetName);
 
 	// Make sure we can create the asset without conflicts
-	if (!CanCreateAsset(AssetName, PackageName, LOCTEXT("CreateANewObject", "Create a new object")))
+	if (!CanCreateAsset(AssetName, PackageName, LOCTEXT("CreateANewObject", "Create a new object"), bOverwriteExisting))
 	{
 		OnCancelled.ExecuteIfBound(Factory);
 		return false;
@@ -2150,7 +2150,7 @@ UObject* UAssetToolsImpl::PerformDuplicateAsset(const FString& AssetName, const
 	const FString PackageName = PackagePath / AssetName;
 
 	// Make sure we can create the asset without conflicts
-	if ( !CanCreateAsset(AssetName, PackageName, LOCTEXT("DuplicateAnObject", "Duplicate an object")) )
+	if ( !CanCreateAsset(AssetName, PackageName, LOCTEXT("DuplicateAnObject", "Duplicate an object"), false) )
 	{
 		return nullptr;
 	}
@@ -4809,7 +4809,7 @@ bool UAssetToolsImpl::CheckForDeletedPackage(const UPackage* Package) const
 	return true;
 }
 
-bool UAssetToolsImpl::CanCreateAsset(const FString& AssetName, const FString& PackageName, const FText& OperationText) const
+bool UAssetToolsImpl::CanCreateAsset(const FString& AssetName, const FString& PackageName, const FText& OperationText, bool bOverwriteExisting) const
 {
 	// @todo: These 'reason' messages are not localized strings!
 	FText Reason;
@@ -4860,28 +4860,32 @@ bool UAssetToolsImpl::CanCreateAsset(const FString& AssetName, const FString& Pa
 	{
 		// Object already exists in either the specified package or another package.
 
-		if (!CanShowDialogs())
+		if (!CanShowDialogs() && !bOverwriteExisting)
 		{
 			UE_LOG(LogAssetTools, Error, TEXT("The asset '%s' already exists in package '%s'. CanCreateAsset cannot ask the user as the application is running unattended and will return false"), *AssetName, *PackageName);
 			return false;
 		}
 		
-		// Check to see if the user wants to replace the object.
-		TSharedRef<SMessageDialog> ConfirmDialog = SNew(SMessageDialog)
-			.Icon(FAppStyle::Get().GetBrush("Icons.WarningWithColor.Large"))
-			.Title(FText(NSLOCTEXT("UnrealEd", "ReplaceExistingObjectInPackageConfirmation_Title", "Overwrite Existing Object")))
-			.Message(FText::Format(NSLOCTEXT("UnrealEd", "ReplaceExistingObjectInPackageConfirmation_Message", "An object already exists with this name.\n\n\tName: {0}\n\tClass: {1}\n\tAsset path: {2}\n\nOverwrite the existing object?"),
-				FText::FromString(AssetName),
-				FText::FromString(ExistingObject->GetClass()->GetName()),
-				FText::FromString(PackageName)))
-			.Buttons({
-				SCustomDialog::FButton(NSLOCTEXT("UnrealEd", "ReplaceExistingObjectInPackageConfirmation_ButtonOverwrite", "Overwrite")).SetPrimary(true),
-				SCustomDialog::FButton(NSLOCTEXT("UnrealEd", "ReplaceExistingObjectInPackageConfirmation_ButtonCancel", "Cancel")),
-				})
-			.ContentMinWidth(300.0f);
-		uint32 ConfirmationResult = ConfirmDialog->ShowModal();
-
-		bool bWantReplace = ConfirmationResult == 0;
+		bool bWantReplace = bOverwriteExisting;
+		if (!bWantReplace)
+		{
+			// Check to see if the user wants to replace the object.
+			TSharedRef<SMessageDialog> ConfirmDialog = SNew(SMessageDialog)
+				.Icon(FAppStyle::Get().GetBrush("Icons.WarningWithColor.Large"))
+				.Title(FText(NSLOCTEXT("UnrealEd", "ReplaceExistingObjectInPackageConfirmation_Title", "Overwrite Existing Object")))
+				.Message(FText::Format(NSLOCTEXT("UnrealEd", "ReplaceExistingObjectInPackageConfirmation_Message", "An object already exists with this name.\n\n\tName: {0}\n\tClass: {1}\n\tAsset path: {2}\n\nOverwrite the existing object?"),
+					FText::FromString(AssetName),
+					FText::FromString(ExistingObject->GetClass()->GetName()),
+					FText::FromString(PackageName)))
+				.Buttons({
+					SCustomDialog::FButton(NSLOCTEXT("UnrealEd", "ReplaceExistingObjectInPackageConfirmation_ButtonOverwrite", "Overwrite")).SetPrimary(true),
+					SCustomDialog::FButton(NSLOCTEXT("UnrealEd", "ReplaceExistingObjectInPackageConfirmation_ButtonCancel", "Cancel")),
+					})
+					.ContentMinWidth(300.0f);
+			uint32 ConfirmationResult = ConfirmDialog->ShowModal();
+
+			bWantReplace = ConfirmationResult == 0;
+		}
 
 		if( bWantReplace )
 		{
diff --git a/Engine/Source/Developer/AssetTools/Private/AssetTools.h b/Engine/Source/Developer/AssetTools/Private/AssetTools.h
--- a/Engine/Source/Developer/AssetTools/Private/AssetTools.h
+++ b/Engine/Source/Developer/AssetTools/Private/AssetTools.h
@@ -84,8 +84,8 @@ public:
 	virtual void UnregisterClassTypeActions(const TSharedRef<IClassTypeActions>& ActionsToRemove) override;
 	virtual void GetClassTypeActionsList( TArray<TWeakPtr<IClassTypeActions>>& OutClassTypeActionsList ) const override;
 	virtual TWeakPtr<IClassTypeActions> GetClassTypeActionsForClass( UClass* Class ) const override;
-	virtual UObject* CreateAsset(const FString& AssetName, const FString& PackagePath, UClass* AssetClass, UFactory* Factory, FName CallingContext = NAME_None) override;
-	virtual bool CreateAssetAsync(const FString& AssetName, const FString& PackagePath, UClass* AssetClass, UFactory* Factory, FAssetCreateComplete OnComplete = FAssetCreateComplete(), FAssetCreateCancelled OnCancelled = FAssetCreateCancelled(), FName CallingContext = NAME_None) override;
+	virtual UObject* CreateAsset(const FString& AssetName, const FString& PackagePath, UClass* AssetClass, UFactory* Factory, FName CallingContext = NAME_None, bool bOverwriteExisting = false) override;
+	virtual bool CreateAssetAsync(const FString& AssetName, const FString& PackagePath, UClass* AssetClass, UFactory* Factory, FAssetCreateComplete OnComplete = FAssetCreateComplete(), FAssetCreateCancelled OnCancelled = FAssetCreateCancelled(), FName CallingContext = NAME_None, bool bOverwriteExisting = false) override;
 	virtual void CreateAssetsFrom(TConstArrayView<UObject*> SourceObjects, UClass* CreateAssetType, const FString& DefaultSuffix, TFunctionRef<UFactory*(UObject*)> FactoryConstructor, FName CallingContext) override;
 	virtual UObject* CreateAssetWithDialog(UClass* AssetClass, UFactory* Factory, FName CallingContext = NAME_None) override;
 	virtual UObject* CreateAssetWithDialog(const FString& AssetName, const FString& PackagePath, UClass* AssetClass, UFactory* Factory, FName CallingContext, const bool bCallConfigureProperties = true) override;
@@ -213,7 +213,7 @@ private:
 	bool CheckForDeletedPackage(const UPackage* Package) const;
 
 	/** Returns true if the supplied Asset name and package are currently valid for creation. */
-	bool CanCreateAsset(const FString& AssetName, const FString& PackageName, const FText& OperationText) const;
+	bool CanCreateAsset(const FString& AssetName, const FString& PackageName, const FText& OperationText, bool bOverwriteExisting) const;
 
 	/** Begins the package migration, after assets have been discovered */
 	void PerformMigratePackages(const TArray<FName>& PackageNamesToMigrate, const FString& DestinationPath, const FMigrationOptions& Options, TSet<FName> AllDependenciesToMigrate, TSet<FName> ExcludedDependencies, bool bSanitizeInputs) const;
diff --git a/Engine/Source/Developer/AssetTools/Public/IAssetTools.h b/Engine/Source/Developer/AssetTools/Public/IAssetTools.h
--- a/Engine/Source/Developer/AssetTools/Public/IAssetTools.h
+++ b/Engine/Source/Developer/AssetTools/Public/IAssetTools.h
@@ -342,10 +342,11 @@ public:
 	 * @param AssetClass the class of the new asset
 	 * @param Factory the factory that will build the new asset
 	 * @param CallingContext optional name of the module or method calling CreateAsset() - this is passed to the factory
+	 * @param bOverwriteExisting if set to true, this will overwrite an existing object without prompting
 	 * @return the new asset or NULL if it fails
 	 */
 	UFUNCTION(BlueprintCallable, Category = "Editor Scripting | Asset Tools")
-	virtual UObject* CreateAsset(const FString& AssetName, const FString& PackagePath, UClass* AssetClass, UFactory* Factory, FName CallingContext = NAME_None) = 0;
+	virtual UObject* CreateAsset(const FString& AssetName, const FString& PackagePath, UClass* AssetClass, UFactory* Factory, FName CallingContext = NAME_None, bool bOverwriteExisting = false) = 0;
 
 	/**
 	 * Creates an asset with the specified name, path, and factory asynchronously
@@ -357,10 +358,11 @@ public:
 	 * @param OnComplete called when the factory has successfully created the asset
 	 * @param OnCancelled called if the asynchronous method is cancelled
 	 * @param CallingContext optional name of the module or method calling CreateAsset() - this is passed to the factory
+	 * @param bOverwriteExisting if set to true, this will overwrite an existing object without prompting
 	 * @return the new asset or NULL if it fails
 	 */
 	UFUNCTION(BlueprintCallable, Category = "Editor Scripting | Asset Tools")
-	virtual bool CreateAssetAsync(const FString& AssetName, const FString& PackagePath, UClass* AssetClass, UFactory* Factory, FAssetCreateCompleteDynamic OnComplete, FAssetCreateCancelledDynamic OnCancelled, FName CallingContext = NAME_None)
+	virtual bool CreateAssetAsync(const FString& AssetName, const FString& PackagePath, UClass* AssetClass, UFactory* Factory, FAssetCreateCompleteDynamic OnComplete, FAssetCreateCancelledDynamic OnCancelled, FName CallingContext = NAME_None, bool bOverwriteExisting = false)
 	{
 		return CreateAssetAsync(
 			AssetName, 
@@ -369,11 +371,12 @@ public:
 			Factory,
 			FAssetCreateComplete::CreateLambda([OnComplete](UFactory* Factory, UObject* NewObject) { OnComplete.ExecuteIfBound(Factory, NewObject); }),
 			FAssetCreateCancelled::CreateLambda([OnCancelled](UFactory* Factory) { OnCancelled.ExecuteIfBound(Factory); }),
-			CallingContext
+			CallingContext,
+			bOverwriteExisting
 		);
 	}
 
-	virtual bool CreateAssetAsync(const FString& AssetName, const FString& PackagePath, UClass* AssetClass, UFactory* Factory, FAssetCreateComplete OnComplete = FAssetCreateComplete(), FAssetCreateCancelled OnCancelled = FAssetCreateCancelled(), FName CallingContext = NAME_None) = 0;
+	virtual bool CreateAssetAsync(const FString& AssetName, const FString& PackagePath, UClass* AssetClass, UFactory* Factory, FAssetCreateComplete OnComplete = FAssetCreateComplete(), FAssetCreateCancelled OnCancelled = FAssetCreateCancelled(), FName CallingContext = NAME_None, bool bOverwriteExisting = false) = 0;
 
 	/**
 	 * Creates one or more assets using the source objects as the basis for the next type.  This is a common enough operation
