Fix a few places to use verify() instead of check() as check compiles out in shipping.
#rnx
#Synced-CL 51587407

diff --git a/Engine/Plugins/Experimental/Workspace/Source/WorkspaceEditor/Private/Workspace.cpp b/Engine/Plugins/Experimental/Workspace/Source/WorkspaceEditor/Private/Workspace.cpp
--- a/Engine/Plugins/Experimental/Workspace/Source/WorkspaceEditor/Private/Workspace.cpp
+++ b/Engine/Plugins/Experimental/Workspace/Source/WorkspaceEditor/Private/Workspace.cpp
@@ -163,7 +163,7 @@ bool UWorkspace::RemoveAsset(const FAssetData& InAsset, bool bSetupUndoRedo, boo
 	if (EntryIndex != INDEX_NONE)
 	{
 		UWorkspaceAssetEntry* EntryToRemove = AssetEntries[EntryIndex];		
-		check(AssetEntries.Remove(EntryToRemove) == 1);
+		verify(AssetEntries.Remove(EntryToRemove) == 1);
 		EntryToRemove->MarkAsGarbage();
 		EntryToRemove->MarkPackageDirty();
 		
diff --git a/Engine/Plugins/VirtualProduction/CameraCalibrationCore/Source/CameraCalibrationCore/Private/CalibratedMapProcessor.cpp b/Engine/Plugins/VirtualProduction/CameraCalibrationCore/Source/CameraCalibrationCore/Private/CalibratedMapProcessor.cpp
--- a/Engine/Plugins/VirtualProduction/CameraCalibrationCore/Source/CameraCalibrationCore/Private/CalibratedMapProcessor.cpp
+++ b/Engine/Plugins/VirtualProduction/CameraCalibrationCore/Source/CameraCalibrationCore/Private/CalibratedMapProcessor.cpp
@@ -94,7 +94,7 @@ void FCalibratedMapProcessor::Update()
 			FDerivedDistortionDataJobArgs PendingJob;
 
 			//Consume Pending job that was peeked
-			check(PendingJobs.Dequeue(PendingJob));
+			verify(PendingJobs.Dequeue(PendingJob));
 
 			NewJob->Output.Reset();
 			NewJob->Output.Focus = PendingJob.Focus;
diff --git a/Engine/Plugins/VirtualProduction/Rivermax/RivermaxCore/Source/RivermaxCore/Private/RivermaxDeviceFinder.cpp b/Engine/Plugins/VirtualProduction/Rivermax/RivermaxCore/Source/RivermaxCore/Private/RivermaxDeviceFinder.cpp
--- a/Engine/Plugins/VirtualProduction/Rivermax/RivermaxCore/Source/RivermaxCore/Private/RivermaxDeviceFinder.cpp
+++ b/Engine/Plugins/VirtualProduction/Rivermax/RivermaxCore/Source/RivermaxCore/Private/RivermaxDeviceFinder.cpp
@@ -88,7 +88,7 @@ namespace UE::RivermaxCore::Private
 				{
 					FDeviceInfo& NewDevice = DevicesCache.FindOrAdd(DeviceIPHash);
 					NewDevice.DeviceIP = Device.InterfaceAddress;
-					check(FIPv4Address::Parse(NewDevice.DeviceIP, NewDevice.Address));
+					verify(FIPv4Address::Parse(NewDevice.DeviceIP, NewDevice.Address));
 					NewDevice.AddressTokens[0] = NewDevice.Address.A;
 					NewDevice.AddressTokens[1] = NewDevice.Address.B;
 					NewDevice.AddressTokens[2] = NewDevice.Address.C;
diff --git a/Engine/Source/Developer/AssetTools/Private/AssetTools.cpp b/Engine/Source/Developer/AssetTools/Private/AssetTools.cpp
--- a/Engine/Source/Developer/AssetTools/Private/AssetTools.cpp
+++ b/Engine/Source/Developer/AssetTools/Private/AssetTools.cpp
@@ -1222,7 +1222,7 @@ namespace UE::AssetTools::Private
 			for (UPackage* Package : ExistingPackageToClean)
 			{
 				const ERenameFlags PkgRenameFlags = REN_AllowPackageLinkerMismatch | REN_DoNotDirty | REN_DontCreateRedirectors | REN_NonTransactional | REN_SkipGeneratedClasses;
-				check(Package->Rename(*MakeUniqueObjectName(nullptr, UPackage::StaticClass(), *FString::Printf(TEXT("%s_DEADFROMMIGRATION"), *Package->GetName())).ToString(), nullptr, PkgRenameFlags));
+				verify(Package->Rename(*MakeUniqueObjectName(nullptr, UPackage::StaticClass(), *FString::Printf(TEXT("%s_DEADFROMMIGRATION"), *Package->GetName())).ToString(), nullptr, PkgRenameFlags));
 				(*PurgeObject)(Package);
 				ForEachObjectWithOuter(Package, *PurgeObject);
 			}
diff --git a/Engine/Source/Developer/AssetTools/Private/PackageMigrationContext.cpp b/Engine/Source/Developer/AssetTools/Private/PackageMigrationContext.cpp
--- a/Engine/Source/Developer/AssetTools/Private/PackageMigrationContext.cpp
+++ b/Engine/Source/Developer/AssetTools/Private/PackageMigrationContext.cpp
@@ -158,13 +158,13 @@ FPackageMigrationContext::FScopedTemporalyMovedPackage::FScopedTemporalyMovedPac
 	, OriginalName(PackageToMove->GetName())
 {
 	const ERenameFlags PkgRenameFlags = REN_DoNotDirty | REN_DontCreateRedirectors | REN_NonTransactional | REN_SkipGeneratedClasses | REN_AllowPackageLinkerMismatch;
-	check(PackageToMove->Rename(*MakeUniqueObjectName(nullptr, UPackage::StaticClass(), *FString::Printf(TEXT("%s_MOVEDFORMIGRATION"), *PackageToMove->GetName())).ToString(), nullptr, PkgRenameFlags));
+	verify(PackageToMove->Rename(*MakeUniqueObjectName(nullptr, UPackage::StaticClass(), *FString::Printf(TEXT("%s_MOVEDFORMIGRATION"), *PackageToMove->GetName())).ToString(), nullptr, PkgRenameFlags));
 }
 
 FPackageMigrationContext::FScopedTemporalyMovedPackage::~FScopedTemporalyMovedPackage()
 {
 	const ERenameFlags PkgRenameFlags = REN_DoNotDirty | REN_DontCreateRedirectors | REN_NonTransactional | REN_SkipGeneratedClasses | REN_AllowPackageLinkerMismatch;
-	check(PackageToMove->Rename(*OriginalName, nullptr, PkgRenameFlags));
+	verify(PackageToMove->Rename(*OriginalName, nullptr, PkgRenameFlags));
 }
 
 const FString& FPackageMigrationContext::FMigrationPackageData::GetInstancedPackageName() const
diff --git a/Engine/Source/Developer/SkeletalMeshUtilitiesCommon/Private/SkeletalMeshNotifier.cpp b/Engine/Source/Developer/SkeletalMeshUtilitiesCommon/Private/SkeletalMeshNotifier.cpp
--- a/Engine/Source/Developer/SkeletalMeshUtilitiesCommon/Private/SkeletalMeshNotifier.cpp
+++ b/Engine/Source/Developer/SkeletalMeshUtilitiesCommon/Private/SkeletalMeshNotifier.cpp
@@ -51,11 +51,11 @@ FSkeletalMeshNotifierBindScope::~FSkeletalMeshNotifierBindScope()
 {
 	if (DelegateToRemoveFromNotifierA.IsValid() && NotifierA.IsValid())
 	{
-		check(NotifierA.Pin()->Delegate().Remove(DelegateToRemoveFromNotifierA));
+		verify(NotifierA.Pin()->Delegate().Remove(DelegateToRemoveFromNotifierA));
 	}
 
 	if (DelegateToRemoveFromNotifierB.IsValid() && NotifierB.IsValid())
 	{
-		check(NotifierB.Pin()->Delegate().Remove(DelegateToRemoveFromNotifierB));
+		verify(NotifierB.Pin()->Delegate().Remove(DelegateToRemoveFromNotifierB));
 	}
 }
diff --git a/Engine/Source/Developer/TraceServices/Private/Model/CookProfilerProvider.cpp b/Engine/Source/Developer/TraceServices/Private/Model/CookProfilerProvider.cpp
--- a/Engine/Source/Developer/TraceServices/Private/Model/CookProfilerProvider.cpp
+++ b/Engine/Source/Developer/TraceServices/Private/Model/CookProfilerProvider.cpp
@@ -128,11 +128,11 @@ void FCookProfilerProvider::CreateAggregation(TArray64<FPackageData>& OutPackage
 					Stack[Depth - 1].ExclTime += CurrentScope.Timestamp - LastTime;
 				}
 
-				check(++Depth < MaxStackSize);
+				verify(++Depth < MaxStackSize);
 			}
 			else
 			{
-				check(--Depth >= 0);
+				verify(--Depth >= 0);
 				check(Stack[Depth].InPackageId == CurrentScope.PackageId);
 				check(Stack[Depth].Type == CurrentScope.Type);
 
diff --git a/Engine/Source/Editor/UnrealEd/Private/Cooker/CookDirector.cpp b/Engine/Source/Editor/UnrealEd/Private/Cooker/CookDirector.cpp
--- a/Engine/Source/Editor/UnrealEd/Private/Cooker/CookDirector.cpp
+++ b/Engine/Source/Editor/UnrealEd/Private/Cooker/CookDirector.cpp
@@ -1639,7 +1639,7 @@ void FCookDirector::TickWorkerShutdowns(ECookDirectorThread TickThread)
 			AbortWorker(NewShutdown->GetWorkerId(), TickThread);
 			{
 				FScopeLock CommunicationScopeLock(&CommunicationLock);
-				check(ShuttingDownWorkers.FindOrAdd(NewShutdown).IsValid()); // Abort worker should have set the value
+				verify(ShuttingDownWorkers.FindOrAdd(NewShutdown).IsValid()); // Abort worker should have set the value
 			}
 			LocalRemoteWorkers.Emplace(NewShutdown);
 		}
diff --git a/Engine/Source/Runtime/Engine/Private/Animation/AnimSequence.cpp b/Engine/Source/Runtime/Engine/Private/Animation/AnimSequence.cpp
--- a/Engine/Source/Runtime/Engine/Private/Animation/AnimSequence.cpp
+++ b/Engine/Source/Runtime/Engine/Private/Animation/AnimSequence.cpp
@@ -5408,8 +5408,8 @@ void UAnimSequence::ReleaseResidency(const ITargetPlatform* InPlatform, uint32 I
 		KeyHash = ResidencyReferencerHashes.FindChecked(InReferencerHash);
 
 		// Remove entry from both maps
-		check(ResidencyReferencerHashes.Remove(InReferencerHash) == 1);
-		check(PlatformHashToReferencers.Remove(KeyHash, InReferencerHash) == 1);
+		verify(ResidencyReferencerHashes.Remove(InReferencerHash) == 1);
+		verify(PlatformHashToReferencers.Remove(KeyHash, InReferencerHash) == 1);
 		
 		// Check if there are any references left
 		PlatformHashToReferencers.MultiFind(KeyHash, ReferencerHashes);
diff --git a/Engine/Source/Runtime/Engine/Private/Animation/PoseAsset.cpp b/Engine/Source/Runtime/Engine/Private/Animation/PoseAsset.cpp
--- a/Engine/Source/Runtime/Engine/Private/Animation/PoseAsset.cpp
+++ b/Engine/Source/Runtime/Engine/Private/Animation/PoseAsset.cpp
@@ -151,7 +151,7 @@ FPoseData* FPoseDataContainer::FindOrAddPoseData(FName PoseName)
 	if (PoseIndex == INDEX_NONE)
 	{
 		PoseIndex = PoseFNames.Add(PoseName);
-		check(PoseIndex == Poses.AddZeroed(1));
+		verify(PoseIndex == Poses.AddZeroed(1));
 	}
 
 	return &Poses[PoseIndex];
