Fix AssetHeaderPatcher crashing when serializing FTextSourceSiteContext::DevNotes

The patcher writes with latest custom versions but still loads with the original versions, this change makes it so it writes with the serialized FFortniteMainBranchObjectVersion (or -1 when not present).

#rnx
#Synced-CL 51800397

diff --git a/Engine/Source/Developer/AssetTools/Private/AssetHeaderPatcher.cpp b/Engine/Source/Developer/AssetTools/Private/AssetHeaderPatcher.cpp
--- a/Engine/Source/Developer/AssetTools/Private/AssetHeaderPatcher.cpp
+++ b/Engine/Source/Developer/AssetTools/Private/AssetHeaderPatcher.cpp
@@ -20,6 +20,7 @@
 #include "Serialization/ImportTypeHierarchy.h"
 #include "UObject/CoreRedirects.h"
 #include "UObject/CoreRedirects/CoreRedirectsContext.h"
+#include "UObject/FortniteMainBranchObjectVersion.h"
 #include "UObject/Linker.h"
 #include "UObject/NameTypes.h"
 #include "UObject/ObjectResource.h"
@@ -2969,7 +2970,8 @@ FAssetHeaderPatcher::EResult FAssetHeaderPatcherInner::PatchHeader_WriteDestinat
 	}
 
 	FNamePatchingWriter Writer(*DstArchive, NameToIndexMap);
-
+	Writer.SetShouldSkipUpdateCustomVersion(true);
+	
 	// set version numbers so components branch correctly
 	Writer.SetUEVer(Summary.GetFileVersionUE());
 	Writer.SetLicenseeUEVer(Summary.GetFileVersionLicenseeUE());
@@ -2980,6 +2982,11 @@ FAssetHeaderPatcher::EResult FAssetHeaderPatcherInner::PatchHeader_WriteDestinat
 		Writer.SetFilterEditorOnly(true);
 	}
 
+	if (Writer.GetCustomVersions().GetVersion(FFortniteMainBranchObjectVersion::GUID) == nullptr)
+	{
+		const_cast<FCustomVersionContainer&>(Writer.GetCustomVersions()).SetVersion(FFortniteMainBranchObjectVersion::GUID, -1, "");
+	}
+
 	int64 LastSectionEndedAt = 0;
 
 	for (int32 SectionIdx = 0; SectionIdx < UE_ARRAY_COUNT(SourceSections); ++SectionIdx)
