[AssetHeaderPatcher]

Corrects an issue in the patcher where substring patching would only patch based on explicit defined a->b transformation passed to the patcher and would not attempt to patch string data representing UE5 paths that can be redirected by a CoreRedirect registered with the engine (rather than with the patcher).

#rnx
#Synced-CL 52136225

diff --git a/Engine/Source/Developer/AssetTools/Internal/AssetHeaderPatcher.h b/Engine/Source/Developer/AssetTools/Internal/AssetHeaderPatcher.h
--- a/Engine/Source/Developer/AssetTools/Internal/AssetHeaderPatcher.h
+++ b/Engine/Source/Developer/AssetTools/Internal/AssetHeaderPatcher.h
@@ -129,7 +129,7 @@ UE_INTERNAL struct FAssetHeaderPatcher
 
 		// String mappings are only used for best-effort replacements. These will be error-prone 
 		// and we should strive for more structured data formats to guard against errors here
-		TMap<FString, FString> StringReplacements;	
+		TMap<FString, FString> VerseStringReplacements;	
 		TMap<FString, FString> StringMountReplacements;
 	};
 
diff --git a/Engine/Source/Developer/AssetTools/Private/AssetHeaderPatcher.cpp b/Engine/Source/Developer/AssetTools/Private/AssetHeaderPatcher.cpp
--- a/Engine/Source/Developer/AssetTools/Private/AssetHeaderPatcher.cpp
+++ b/Engine/Source/Developer/AssetTools/Private/AssetHeaderPatcher.cpp
@@ -631,10 +631,6 @@ void FAssetHeaderPatcher::FContext::GenerateAdditionalRemappings()
 		const FCoreRedirectObjectName& SrcName = Redirect.OldName;
 		const FCoreRedirectObjectName& DstName = Redirect.NewName;
 
-		// Do not include Src->Dst ObjectName mappings since it's too likely that we will incorrectly rename when dealing with string data
-		StringReplacements.Add(SrcName.PackageName.ToString(), DstName.PackageName.ToString());
-		StringReplacements.Add(SrcName.ToString(), DstName.ToString());
-
 		// Tag data can contain VersePaths which are like Top-Level Asset Paths
 		// but with a mountpoint prefix and only '/' delimiters
 		for (FString& VerseMount : VerseMountPoints)
@@ -652,7 +648,7 @@ void FAssetHeaderPatcher::FContext::GenerateAdditionalRemappings()
 			DstName.PackageName.AppendString(DstNameBuilder);
 			DstNameBuilder.AppendChar(TEXT('/'));
 			DstName.ObjectName.AppendString(DstNameBuilder);
-			StringReplacements.Add(SrcNameBuilder.ToString(), DstNameBuilder.ToString());
+			VerseStringReplacements.Add(SrcNameBuilder.ToString(), DstNameBuilder.ToString());
 		}
 	}
 
@@ -928,10 +924,10 @@ public:
 		int32 FileOffset = 0;
 	};
 
-	FAssetHeaderPatcherInner(const FString& InSrcAsset, const FString& InDstAsset, const TMap<FString, FString>& InStringReplacements, const TMap<FString, FString>& InStringMountPointReplacements, FArchive* InDstArchive = nullptr)
+	FAssetHeaderPatcherInner(const FString& InSrcAsset, const FString& InDstAsset, const TMap<FString, FString>& InVerseStringReplacements, const TMap<FString, FString>& InStringMountPointReplacements, FArchive* InDstArchive = nullptr)
 		: SrcAsset(InSrcAsset)
 		, DstAsset(InDstAsset)
-		, StringReplacements(InStringReplacements)
+		, VerseStringReplacements(InVerseStringReplacements)
 		, StringMountPointReplacements(InStringMountPointReplacements)
 		, DstArchive(InDstArchive)
 		, bPatchPrimaryAssetTag(false)
@@ -1013,7 +1009,7 @@ public:
 
 	const FString& SrcAsset;
 	const FString& DstAsset;
-	const TMap<FString, FString>& StringReplacements;
+	const TMap<FString, FString>& VerseStringReplacements;
 	const TMap<FString, FString>& StringMountPointReplacements;
 	FArchive* DstArchive = nullptr;
 	TUniquePtr<FArchive> DstArchiveOwner;
@@ -1087,7 +1083,7 @@ public:
 
 FAssetHeaderPatcher::EResult FAssetHeaderPatcher::DoPatch(const FString& InSrcAsset, const FString& InDstAsset, const FContext& InContext)
 {
-	FAssetHeaderPatcherInner Inner(InSrcAsset, InDstAsset, InContext.StringReplacements, InContext.StringMountReplacements);
+	FAssetHeaderPatcherInner Inner(InSrcAsset, InDstAsset, InContext.VerseStringReplacements, InContext.StringMountReplacements);
 	if (!FFileHelper::LoadFileToArray(Inner.SrcBuffer, *Inner.SrcAsset))
 	{
 		UE_LOG(LogAssetHeaderPatcher, Error, TEXT("Failed to load %s"), *Inner.SrcAsset);
@@ -1688,15 +1684,37 @@ bool FAssetHeaderPatcherInner::ShouldReplaceMountPoint(const FStringView InPath,
 // error-prone and should be avoided.
 bool FAssetHeaderPatcherInner::DoPatch(FString& InOutString)
 {
+	auto RedirectString = [this](const FStringView InMaybePath, FString& OutRedirectedString)
+		{
+			// Strings we will redirect can be made of Package, Outer and Object FNames so we can't handle anything larger than 3 x NAME_SIZE
+			if (InMaybePath.Len() >= (NAME_SIZE*3))
+			{
+				return false;
+			}
+
+			FCoreRedirectObjectName SrcName = FString(InMaybePath);
+			FCoreRedirectObjectName DstName = FCoreRedirects::GetRedirectedName(ECoreRedirectFlags::Type_AllMask, SrcName, ECoreRedirectMatchFlags::DisallowPartialLHSMatch);
+			if (SrcName != DstName)
+			{
+				OutRedirectedString = DstName.ToString();
+				return true;
+			}
+			else
+			{
+				FStringView MaybeReplacement = Find(VerseStringReplacements, InMaybePath);
+				if (!MaybeReplacement.IsEmpty())
+				{
+					OutRedirectedString = MaybeReplacement;
+					return true;
+				}
+			}
+			return false;
+		};
+
 	// Attempt a direct replacement
+	if (RedirectString(InOutString, InOutString))
 	{
-		// Find a Path, change a Path.
-		FStringView MaybeReplacement = Find(StringReplacements, InOutString);
-		if (!MaybeReplacement.IsEmpty())
-		{
-			InOutString = MaybeReplacement;
-			return true;
-		}
+		return true;
 	}
 
 	// Direct replacement failed so now try substring replacements
@@ -1758,9 +1776,9 @@ bool FAssetHeaderPatcherInner::DoPatch(FString& InOutString)
 
 			FStringView SrcMountPoint;
 			FStringView DstMountPoint;
-			FStringView ObjectPath(PathView.GetData() + LastPathDelimiterPos, ColonPos - LastPathDelimiterPos);
-			FStringView MaybeReplacement = Find(StringReplacements, ObjectPath);
-			if (!MaybeReplacement.IsEmpty())
+			FStringView MaybePath(PathView.GetData() + LastPathDelimiterPos, ColonPos - LastPathDelimiterPos);
+			FString MaybeReplacement;
+			if (RedirectString(MaybePath, MaybeReplacement))
 			{
 				FStringView LeftPart(*InOutString, int32(PathView.GetData() - *InOutString) + LastPathDelimiterPos);
 				FStringView RightPart(PathView.GetData() + ColonPos);
@@ -1776,7 +1794,7 @@ bool FAssetHeaderPatcherInner::DoPatch(FString& InOutString)
 				// Keep searching until the path is depleted since there might be more than one path to replace
 				PathView = FStringView(*InOutString + LeftPart.Len() + MaybeReplacement.Len() + 1);
 			}
-			else if (ShouldReplaceMountPoint(ObjectPath, SrcMountPoint, DstMountPoint))
+			else if (ShouldReplaceMountPoint(MaybePath, SrcMountPoint, DstMountPoint))
 			{
 				FStringView LeftPart(*InOutString, int32(PathView.GetData() - *InOutString) + LastPathDelimiterPos);
 				FStringView RightPart(PathView.GetData() + LastPathDelimiterPos + SrcMountPoint.Len());
@@ -1804,7 +1822,7 @@ bool FAssetHeaderPatcherInner::DoPatch(FString& InOutString)
 	{
 		// Patch quoted paths.
 		// Path occurs to the right of the first "'" or """ 
-		auto PatchQuotedPath = [this, &DstStringBuilder](FString& StringToPatch, FStringView Quote)
+		auto PatchQuotedPath = [this, &DstStringBuilder, &RedirectString](FString& StringToPatch, FStringView Quote)
 			{
 				int32 FirstQuotePos = INDEX_NONE;
 				bool bFoundReplacement = false;
@@ -1820,9 +1838,9 @@ bool FAssetHeaderPatcherInner::DoPatch(FString& InOutString)
 
 					FStringView SrcMountPoint;
 					FStringView DstMountPoint;
-					FStringView StrippedQuotedPath = FStringView(PathView.GetData() + FirstQuotePos + 1, SecondQuotePos - FirstQuotePos - 1); // +1 and -1 are to skip the quotes
-					FStringView MaybeReplacement = Find(StringReplacements, StrippedQuotedPath);
-					if (!MaybeReplacement.IsEmpty())
+					FStringView MaybePath = FStringView(PathView.GetData() + FirstQuotePos + 1, SecondQuotePos - FirstQuotePos - 1); // +1 and -1 are to skip the quotes
+					FString MaybeReplacement;
+					if (RedirectString(MaybePath, MaybeReplacement))
 					{
 						FStringView LeftPart(*StringToPatch, int32(PathView.GetData() - *StringToPatch) + FirstQuotePos + 1); // +1 to ensure we include the quote
 						FStringView RightPart(PathView.GetData() + SecondQuotePos);
@@ -1838,7 +1856,7 @@ bool FAssetHeaderPatcherInner::DoPatch(FString& InOutString)
 						// Keep searching until the path is depleted since there might be more than one path to replace
 						PathView = FStringView(*StringToPatch + LeftPart.Len() + MaybeReplacement.Len() + 1);
 					}
-					else if (ShouldReplaceMountPoint(StrippedQuotedPath, SrcMountPoint, DstMountPoint))
+					else if (ShouldReplaceMountPoint(MaybePath, SrcMountPoint, DstMountPoint))
 					{
 						FStringView LeftPart(*StringToPatch, int32(PathView.GetData() - *StringToPatch) + FirstQuotePos + 1); // +1 to ensure we include the quote
 						FStringView RightPart(PathView.GetData() + FirstQuotePos + SrcMountPoint.Len() + 1); // +1 to ensure we skip the first quote
@@ -1853,7 +1871,7 @@ bool FAssetHeaderPatcherInner::DoPatch(FString& InOutString)
 
 						// Keep searching until the path is depleted since there might be more than one path to replace
 						// Skip to the end quote since we know we didn't have any matches within the quotes beyond the mount
-						PathView = FStringView(*StringToPatch + SecondQuotePos + 1 +(DstMountPoint.Len() - SrcMountPoint.Len())); // Dst - Src to account for new SecondQuotePos after replacement
+						PathView = FStringView(*StringToPatch + SecondQuotePos + 1 + (DstMountPoint.Len() - SrcMountPoint.Len())); // Dst - Src to account for new SecondQuotePos after replacement
 					}
 					else
 					{
@@ -3844,9 +3862,9 @@ TEST_CASE_NAMED(FAssetHeaderPatcherTests, "AssetHeaderPatcher", "[AssetHeaderPat
 	struct FTestPatcherContext : FAssetHeaderPatcher::FContext
 	{
 		FTestPatcherContext(TMap<FString, FString> PackageRenameMap, bool bGatherDependentPackages = true) : FAssetHeaderPatcher::FContext(PackageRenameMap, bGatherDependentPackages) {}
-		const TMap<FString, FString>& GetStringReplacements()
+		const TMap<FString, FString>& GetVerseStringReplacements()
 		{
-			return StringReplacements;
+			return VerseStringReplacements;
 		}
 
 		void GenerateRemappings()
@@ -4152,11 +4170,17 @@ TEST_CASE_NAMED(FAssetHeaderPatcherTests, "AssetHeaderPatcher", "[AssetHeaderPat
 	ON_SCOPE_EXIT{ FCoreRedirectsContext::SetThreadContext(OriginalContext); };
 
 	FTestPatcherContext Context(PackageRenameMap, false /*bGatherDependentPackages*/);
-	const TMap<FString, FString>& StringReplacements = Context.GetStringReplacements();
-	CHECK(StringReplacements.Num() > PackageRenameMap.Num()); // Ensure we generated more mappings off of the PackageRenameMap
+	const TMap<FString, FString>& VerseStringReplacements = Context.GetVerseStringReplacements();
+	CHECK(VerseStringReplacements.Num() > PackageRenameMap.Num()); // Ensure we generated more mappings off of the PackageRenameMap
 	CHECK(FCoreRedirects::AddRedirectList(Context.GetRedirects(), TEXT("Asset Header Patcher Tests")));
 
-	FAssetHeaderPatcherInner Patcher(DummySrcDstAsset, DummySrcDstAsset, StringReplacements, MountPointReplacementMap);
+	// Add Global CoreRedirects (scoped within the test context)
+	{
+		FCoreRedirect TagClassRedirect(ECoreRedirectFlags::Type_Class, TEXT("/Script/VerseGameplayTags.VerseTagMarkupComponent"), TEXT("/Script/VerseTags.VerseTagMarkupComponent"));
+		FCoreRedirects::AddRedirectList({ TagClassRedirect }, TEXT("Asset Header Patcher Tests - Additional Redirects"));
+	}
+
+	FAssetHeaderPatcherInner Patcher(DummySrcDstAsset, DummySrcDstAsset, VerseStringReplacements, MountPointReplacementMap);
 
 	auto AddToNameTable = [&Patcher](FName Name)
 		{
@@ -4588,6 +4612,15 @@ TEST_CASE_NAMED(FAssetHeaderPatcherTests, "AssetHeaderPatcher", "[AssetHeaderPat
 					CHECK(Patcher.DoPatch(Actual));
 					CHECK_EQUALS(TEXT("Patch substring with quoted package, object and sub-object paths. No non-fully-qualified ObjectNames are patched."), Actual, Expected);
 				}
+
+				// Relies on the TagClassRedirect defined above
+				{
+					FString Actual(TEXT(R"(CacheVersion=1,AllowedClassReferences=("/Script/VerseGameplayTags.VerseTagMarkupComponent"))"));
+					FString Expected(TEXT(R"(CacheVersion=1,AllowedClassReferences=("/Script/VerseTags.VerseTagMarkupComponent"))"));
+					CHECK(Patcher.DoPatch(Actual));
+					CHECK_EQUALS(TEXT("Patch substring with quoted package, object and sub-object paths. No non-fully-qualified ObjectNames are patched."), Actual, Expected);
+					
+				}
 			}
 		}
 
