EditorUsability : AssetTools
- Fixed delete empty folder without dialog confirmation to avoid affecting tests behavior, see test : Editor.ContentBrowser.Folders.FContentBrowserFolderCreation.Verify_Folder_Can_Be_Created_Inside_Folder

#Synced-CL 52573604

diff --git a/Engine/Source/Developer/AssetTools/Private/AssetViewUtils.cpp b/Engine/Source/Developer/AssetTools/Private/AssetViewUtils.cpp
--- a/Engine/Source/Developer/AssetTools/Private/AssetViewUtils.cpp
+++ b/Engine/Source/Developer/AssetTools/Private/AssetViewUtils.cpp
@@ -452,9 +452,10 @@ bool AssetViewUtils::DeleteFolders(const FDeleteFolderParameters& InOptions)
 	bool bSuccess = false;
 	if ( NumAssetsInPaths == 0 )
 	{
-		// There were no assets, allow the folder delete.
+		// There were no assets found, allow the empty folders deletion without any confirmation
 		ObjectTools::FDeleteItemsParameters Options;
 		Options.FoldersToDelete = InOptions.PathsToDelete;
+		Options.bShowConfirmation = false;
 		bSuccess = ObjectTools::DeleteItems(Options).FoldersDeleted == InOptions.PathsToDelete.Num();
 	}
 	else
