[AssetTool] Fix wrong permission check for the CB Add Menu

#Synced-CL 52619962

diff --git a/Engine/Source/Developer/AssetTools/Private/AssetTools.cpp b/Engine/Source/Developer/AssetTools/Private/AssetTools.cpp
--- a/Engine/Source/Developer/AssetTools/Private/AssetTools.cpp
+++ b/Engine/Source/Developer/AssetTools/Private/AssetTools.cpp
@@ -6381,17 +6381,17 @@ bool UAssetToolsImpl::IsAssetClassSupported(const UClass* AssetClass) const
 	TWeakPtr<IAssetTypeActions> AssetTypeActions = GetAssetTypeActionsForClass(AssetClass);
 	if (AssetTypeActions.IsValid())
 	{
-		bool bSupported = false;
-		if (const UClass* SupportedClass = AssetTypeActions.Pin()->GetSupportedClass())
+		const TSharedRef<FPathPermissionList>& CreateAssetPermissionList = GetAssetClassPathPermissionList(EAssetClassAction::CreateAsset);
+		if (CreateAssetPermissionList->HasFiltering())
 		{
-			bSupported = GetAssetClassPathPermissionList(EAssetClassAction::CreateAsset)->PassesFilter(SupportedClass->GetClassPathName().ToString());
+			// Use the Real AssetClass passed from the Factory and not the one from AssetDefinition.
+			// This way we avoid issue where the real AssetClass falls back to use the AssetDefinition of its parent.
+			FNameBuilder AssetClassPath;
+			AssetClass->GetClassPathName().AppendString(AssetClassPath);
+			return CreateAssetPermissionList->PassesFilter(AssetClassPath.ToView());	
 		}
-		//else
-		//{
-		//	bSupported = !AssetTypeActions.Pin()->GetFilterName().IsNone();
-		//}
 
-		return bSupported;
+		return true;
 	}
 	
 	return false;
